fixed double click event on AI graph subnodes
#ue4

diff --git a/Engine/Source/Editor/AIGraph/Private/SGraphNodeAI.cpp b/Engine/Source/Editor/AIGraph/Private/SGraphNodeAI.cpp
--- a/Engine/Source/Editor/AIGraph/Private/SGraphNodeAI.cpp
+++ b/Engine/Source/Editor/AIGraph/Private/SGraphNodeAI.cpp
@@ -345,6 +345,7 @@ void SGraphNodeAI::SetOwner(const TSharedRef<SGraphPanel>& OwnerPanel)
 		if (ChildWidget.IsValid())
 		{
 			ChildWidget->SetOwner(OwnerPanel);
+			OwnerPanel->AttachGraphEvents(ChildWidget);
 		}
 	}
 }
