Fix deleted behavior tree blueprint nodes showing up in context menu

By the time assets are deleted, ResolveName() will not return the object name as the underlying package is no longer there, so it will not get deleted form the class tree. Performing direct string manupulation on the object path avoids this.

UE-8201 - Task node blueprints remain active in the context menu after they are deleted.

diff --git a/Engine/Source/Editor/AIGraph/Private/AIGraphTypes.cpp b/Engine/Source/Editor/AIGraph/Private/AIGraphTypes.cpp
--- a/Engine/Source/Editor/AIGraph/Private/AIGraphTypes.cpp
+++ b/Engine/Source/Editor/AIGraph/Private/AIGraphTypes.cpp
@@ -288,8 +288,8 @@ void FGraphNodeClassHelper::OnAssetRemoved(const class FAssetData& AssetData)
 	if (GeneratedClassname)
 	{
 		FString AssetClassName = *GeneratedClassname;
-		UObject* Outer1(NULL);
-		ResolveName(Outer1, AssetClassName, false, false);
+		ConstructorHelpers::StripObjectClass(AssetClassName);
+		AssetClassName = FPackageName::ObjectPathToObjectName(*AssetClassName);
 
 		TSharedPtr<FGraphNodeClassNode> Node = FindBaseClassNode(RootNode, AssetClassName);
 		if (Node.IsValid() && Node->ParentNode.IsValid())
