fixed order of behavior tree nodes when creating a new node from existing pin
#ue4 UE-17114

diff --git a/Engine/Source/Editor/AIGraph/Private/AIGraphSchema.cpp b/Engine/Source/Editor/AIGraph/Private/AIGraphSchema.cpp
--- a/Engine/Source/Editor/AIGraph/Private/AIGraphSchema.cpp
+++ b/Engine/Source/Editor/AIGraph/Private/AIGraphSchema.cpp
@@ -39,8 +39,6 @@ UEdGraphNode* FAISchemaAction_NewNode::PerformAction(class UEdGraph* ParentGraph
 
 		NodeTemplate->CreateNewGuid();
 		NodeTemplate->PostPlacedNewNode();
-		NodeTemplate->AllocateDefaultPins();
-		NodeTemplate->AutowireNewNode(FromPin);
 
 		// For input pins, new node will generally overlap node being dragged off
 		// Work out if we want to visually push away from connected node
@@ -62,6 +60,10 @@ UEdGraphNode* FAISchemaAction_NewNode::PerformAction(class UEdGraph* ParentGraph
 		NodeTemplate->NodePosY = Location.Y;
 		NodeTemplate->SnapToGrid(SNAP_GRID);
 
+		// setup pins after placing node in correct spot, since pin sorting will happen as soon as link connection change occurs
+		NodeTemplate->AllocateDefaultPins();
+		NodeTemplate->AutowireNewNode(FromPin);
+
 		ResultNode = NodeTemplate;
 	}
 
