Cleaning out programmer-only comments that were harvested by Doxygen.

diff --git a/Engine/Plugins/2D/Paper2D/Source/Paper2D/Classes/PaperFlipbook.h b/Engine/Plugins/2D/Paper2D/Source/Paper2D/Classes/PaperFlipbook.h
--- a/Engine/Plugins/2D/Paper2D/Source/Paper2D/Classes/PaperFlipbook.h
+++ b/Engine/Plugins/2D/Paper2D/Source/Paper2D/Classes/PaperFlipbook.h
@@ -135,13 +135,13 @@ public:
 	// Rebuilds cached data about the animation (such as total number of frames that the keyframes span, etc...)
 	void InvalidateCachedData();
 
-	// UObject interface
+	//~ Begin UObject Interface
 	virtual void Serialize(FArchive& Ar) override;
 	virtual void PostLoad() override;
 #if WITH_EDITOR
 	virtual void PostEditChangeProperty(FPropertyChangedEvent& PropertyChangedEvent) override;
 #endif
-	// End of UObject interface
+	//~ End UObject Interface
 
 	// Returns true if the flipbook asset contains the specified sprite asset in any frames
 	bool ContainsSprite(UPaperSprite* SpriteAsset) const;
diff --git a/Engine/Plugins/2D/Paper2D/Source/Paper2DEditor/Classes/PaperFlipbookActorFactory.h b/Engine/Plugins/2D/Paper2D/Source/Paper2DEditor/Classes/PaperFlipbookActorFactory.h
--- a/Engine/Plugins/2D/Paper2D/Source/Paper2DEditor/Classes/PaperFlipbookActorFactory.h
+++ b/Engine/Plugins/2D/Paper2D/Source/Paper2DEditor/Classes/PaperFlipbookActorFactory.h
@@ -9,9 +9,9 @@ class UPaperFlipbookActorFactory : public UActorFactory
 {
 	GENERATED_UCLASS_BODY()
 
-	// Begin UActorFactory Interface
+	//~ Begin UActorFactory Interface
 	virtual void PostSpawnActor(UObject* Asset, AActor* NewActor) override;
 	virtual void PostCreateBlueprint(UObject* Asset, AActor* CDO) override;
 	virtual bool CanCreateActorFrom(const FAssetData& AssetData, FText& OutErrorMsg) override;
-	// End UActorFactory Interface
+	//~ End UActorFactory Interface
 };
diff --git a/Engine/Plugins/2D/Paper2D/Source/Paper2DEditor/Classes/PaperSpriteActorFactory.h b/Engine/Plugins/2D/Paper2D/Source/Paper2DEditor/Classes/PaperSpriteActorFactory.h
--- a/Engine/Plugins/2D/Paper2D/Source/Paper2DEditor/Classes/PaperSpriteActorFactory.h
+++ b/Engine/Plugins/2D/Paper2D/Source/Paper2DEditor/Classes/PaperSpriteActorFactory.h
@@ -9,9 +9,9 @@ class UPaperSpriteActorFactory : public UActorFactory
 {
 	GENERATED_UCLASS_BODY()
 
-	// Begin UActorFactory Interface
+	//~ Begin UActorFactory Interface
 	virtual void PostSpawnActor(UObject* Asset, AActor* NewActor) override;
 	virtual void PostCreateBlueprint(UObject* Asset, AActor* CDO) override;
 	virtual bool CanCreateActorFrom(const FAssetData& AssetData, FText& OutErrorMsg) override;
-	// End UActorFactory Interface
+	//~ End UActorFactory Interface
 };
diff --git a/Engine/Plugins/2D/Paper2D/Source/PaperSpriteSheetImporter/Private/PaperSpriteSheetReimportFactory.h b/Engine/Plugins/2D/Paper2D/Source/PaperSpriteSheetImporter/Private/PaperSpriteSheetReimportFactory.h
--- a/Engine/Plugins/2D/Paper2D/Source/PaperSpriteSheetImporter/Private/PaperSpriteSheetReimportFactory.h
+++ b/Engine/Plugins/2D/Paper2D/Source/PaperSpriteSheetImporter/Private/PaperSpriteSheetReimportFactory.h
@@ -11,10 +11,10 @@ class UPaperSpriteSheetReimportFactory : public UPaperSpriteSheetImportFactory,
 {
 	GENERATED_UCLASS_BODY()
 
-	// Begin FReimportHandler interface
+	//~ Begin FReimportHandler Interface
 	virtual bool CanReimport(UObject* Obj, TArray<FString>& OutFilenames) override;
 	virtual void SetReimportPaths(UObject* Obj, const TArray<FString>& NewReimportPaths) override;
 	virtual EReimportResult::Type Reimport(UObject* Obj) override;
 	virtual int32 GetPriority() const override;
-	// End FReimportHandler interface
+	//~ End FReimportHandler Interface
 };
diff --git a/Engine/Plugins/Editor/SpeedTreeImporter/Source/SpeedTreeImporter/Classes/ReimportSpeedTreeFactory.h b/Engine/Plugins/Editor/SpeedTreeImporter/Source/SpeedTreeImporter/Classes/ReimportSpeedTreeFactory.h
--- a/Engine/Plugins/Editor/SpeedTreeImporter/Source/SpeedTreeImporter/Classes/ReimportSpeedTreeFactory.h
+++ b/Engine/Plugins/Editor/SpeedTreeImporter/Source/SpeedTreeImporter/Classes/ReimportSpeedTreeFactory.h
@@ -13,12 +13,12 @@ class UReimportSpeedTreeFactory : public USpeedTreeImportFactory, public FReimpo
 {
 	GENERATED_UCLASS_BODY()
 
-	// Begin FReimportHandler interface
+	//~ Begin FReimportHandler interface
 	virtual bool CanReimport(UObject* Obj, TArray<FString>& OutFilenames) override;
 	virtual void SetReimportPaths(UObject* Obj, const TArray<FString>& NewReimportPaths) override;
 	virtual EReimportResult::Type Reimport(UObject* Obj) override;
 	virtual int32 GetPriority() const override;
-	// End FReimportHandler interface
+	//~ End FReimportHandler interface
 };
 
 
diff --git a/Engine/Plugins/Editor/SpeedTreeImporter/Source/SpeedTreeImporter/Classes/SpeedTreeImportFactory.h b/Engine/Plugins/Editor/SpeedTreeImporter/Source/SpeedTreeImporter/Classes/SpeedTreeImportFactory.h
--- a/Engine/Plugins/Editor/SpeedTreeImporter/Source/SpeedTreeImporter/Classes/SpeedTreeImportFactory.h
+++ b/Engine/Plugins/Editor/SpeedTreeImporter/Source/SpeedTreeImporter/Classes/SpeedTreeImportFactory.h
@@ -12,7 +12,7 @@ class USpeedTreeImportFactory : public UFactory
 {
 	GENERATED_UCLASS_BODY()
 
-	// Begin UFactory Interface
+	//~ Begin UFactory Interface
 	virtual FText GetDisplayName() const override;
 #if WITH_SPEEDTREE
 	virtual bool DoesSupportClass(UClass * Class) override;
diff --git a/Engine/Plugins/Experimental/HTML5Networking/Source/HTML5Networking/Classes/WebSocketConnection.h b/Engine/Plugins/Experimental/HTML5Networking/Source/HTML5Networking/Classes/WebSocketConnection.h
--- a/Engine/Plugins/Experimental/HTML5Networking/Source/HTML5Networking/Classes/WebSocketConnection.h
+++ b/Engine/Plugins/Experimental/HTML5Networking/Source/HTML5Networking/Classes/WebSocketConnection.h
@@ -12,7 +12,7 @@ class HTML5NETWORKING_API UWebSocketConnection : public UNetConnection
 
 	class FWebSocket*			WebSocket; 
 
-	// Begin NetConnection Interface
+	//~ Begin NetConnection Interface
 	virtual void InitBase(UNetDriver* InDriver, class FSocket* InSocket, const FURL& InURL, EConnectionState InState, int32 InMaxPacket = 0, int32 InPacketOverhead = 0) override;
 	virtual void InitRemoteConnection(UNetDriver* InDriver, class FSocket* InSocket, const FURL& InURL, const class FInternetAddr& InRemoteAddr, EConnectionState InState, int32 InMaxPacket = 0, int32 InPacketOverhead = 0) override;
 	virtual void InitLocalConnection(UNetDriver* InDriver, class FSocket* InSocket, const FURL& InURL, EConnectionState InState, int32 InMaxPacket = 0, int32 InPacketOverhead = 0) override;
@@ -21,7 +21,7 @@ class HTML5NETWORKING_API UWebSocketConnection : public UNetConnection
 	FString LowLevelDescribe() override;
 	virtual void Tick();
 	virtual void FinishDestroy();
-	// End NetConnection Interface
+	//~ End NetConnection Interface
 
 
 	void SetWebSocket(FWebSocket* InWebSocket);
diff --git a/Engine/Plugins/Experimental/HTML5Networking/Source/HTML5Networking/Classes/WebSocketNetDriver.h b/Engine/Plugins/Experimental/HTML5Networking/Source/HTML5Networking/Classes/WebSocketNetDriver.h
--- a/Engine/Plugins/Experimental/HTML5Networking/Source/HTML5Networking/Classes/WebSocketNetDriver.h
+++ b/Engine/Plugins/Experimental/HTML5Networking/Source/HTML5Networking/Classes/WebSocketNetDriver.h
@@ -15,7 +15,7 @@ class HTML5NETWORKING_API UWebSocketNetDriver : public UNetDriver
 	UPROPERTY(Config)
 	int32 WebSocketPort;
 
-	// Begin UNetDriver interface.
+	//~ Begin UNetDriver Interface.
 	virtual bool IsAvailable() const override;
 	virtual bool InitBase(bool bInitAsClient, FNetworkNotify* InNotify, const FURL& URL, bool bReuseAddressAndPort, FString& Error) override;
 	virtual bool InitConnect(FNetworkNotify* InNotify, const FURL& ConnectURL, FString& Error) override;
@@ -30,11 +30,11 @@ class HTML5NETWORKING_API UWebSocketNetDriver : public UNetDriver
 	virtual class ISocketSubsystem* GetSocketSubsystem() override;
 	virtual FSocket * CreateSocket();
 
-	// End UNetDriver interface.
+	//~ End UNetDriver Interface.
 
-	// Begin FExec Interface
+	//~ Begin FExec Interface
 	virtual bool Exec(UWorld* InWorld, const TCHAR* Cmd, FOutputDevice& Ar = *GLog) override;
-	// End FExec Interface
+	//~ End FExec Interface
 
 	/**
 	* Exec command handlers
diff --git a/Engine/Plugins/Experimental/LiveEditor/Source/LiveEditor/Classes/K2Node_LiveEditObject.h b/Engine/Plugins/Experimental/LiveEditor/Source/LiveEditor/Classes/K2Node_LiveEditObject.h
--- a/Engine/Plugins/Experimental/LiveEditor/Source/LiveEditor/Classes/K2Node_LiveEditObject.h
+++ b/Engine/Plugins/Experimental/LiveEditor/Source/LiveEditor/Classes/K2Node_LiveEditObject.h
@@ -12,19 +12,19 @@ class UK2Node_LiveEditObject : public UK2Node
 	GENERATED_UCLASS_BODY()
 
 #if WITH_EDITOR
-	// Begin UEdGraphNode interface.
+	//~ Begin UEdGraphNode Interface.
 	virtual void AllocateDefaultPins() override;
 	virtual FText GetNodeTitle(ENodeTitleType::Type TitleType) const override;
 	virtual void PinDefaultValueChanged(UEdGraphPin* Pin) override;
-	// End UEdGraphNode interface.
+	//~ End UEdGraphNode Interface.
 
-	// Begin UK2Node interface
+	//~ Begin UK2Node Interface
 	virtual bool IsNodeSafeToIgnore() const override { return true; }
 	virtual void ReallocatePinsDuringReconstruction(TArray<UEdGraphPin*>& OldPins) override;
 	virtual class FNodeHandlingFunctor* CreateNodeHandler(class FKismetCompilerContext& CompilerContext) const override;
 	virtual void ExpandNode(class FKismetCompilerContext& CompilerContext, UEdGraph* SourceGraph) override;
 	virtual void GetNodeAttributes( TArray<TKeyValuePair<FString, FString>>& OutNodeAttributes ) const override;
-	// End UK2Node interface
+	//~ End UK2Node Interface
 
 	/** Create new pins to show LiveEditable properties on archetype */
 	void CreatePinsForClass(UClass* InClass);
diff --git a/Engine/Plugins/Experimental/Phya/Source/Phya/Classes/PhyaCollisionHandler.h b/Engine/Plugins/Experimental/Phya/Source/Phya/Classes/PhyaCollisionHandler.h
--- a/Engine/Plugins/Experimental/Phya/Source/Phya/Classes/PhyaCollisionHandler.h
+++ b/Engine/Plugins/Experimental/Phya/Source/Phya/Classes/PhyaCollisionHandler.h
@@ -32,10 +32,10 @@ class UPhyaCollisionHandler : public UPhysicsCollisionHandler
 {
 	GENERATED_UCLASS_BODY()
 
-	// Begin UPhysicsCollisionHandler interface
+	//~ Begin UPhysicsCollisionHandler Interface
 	virtual void InitCollisionHandler() override;
 	virtual void HandlePhysicsCollisions_AssumesLocked(const TArray<FCollisionNotifyInfo>& PendingCollisionNotifies) override;
-	// End UPhysicsCollisionHandler interface
+	//~ End UPhysicsCollisionHandler Interface
 
 	void TestImpact();
 
diff --git a/Engine/Plugins/Runtime/AndroidDeviceProfileSelector/Source/AndroidDeviceProfileSelector/Private/AndroidDeviceProfileSelectorModule.h b/Engine/Plugins/Runtime/AndroidDeviceProfileSelector/Source/AndroidDeviceProfileSelector/Private/AndroidDeviceProfileSelectorModule.h
--- a/Engine/Plugins/Runtime/AndroidDeviceProfileSelector/Source/AndroidDeviceProfileSelector/Private/AndroidDeviceProfileSelectorModule.h
+++ b/Engine/Plugins/Runtime/AndroidDeviceProfileSelector/Source/AndroidDeviceProfileSelector/Private/AndroidDeviceProfileSelectorModule.h
@@ -11,15 +11,15 @@ class FAndroidDeviceProfileSelectorModule
 {
 public:
 
-	// Begin IDeviceProfileSelectorModule interface
+	//~ Begin IDeviceProfileSelectorModule Interface
 	virtual const FString GetRuntimeDeviceProfileName() override;
-	// End IDeviceProfileSelectorModule interface
+	//~ End IDeviceProfileSelectorModule Interface
 
 
-	// Begin IModuleInterface interface
+	//~ Begin IModuleInterface Interface
 	virtual void StartupModule() override;
 	virtual void ShutdownModule() override;
-	// End IModuleInterface interface
+	//~ End IModuleInterface Interface
 
 	
 	/**
diff --git a/Engine/Plugins/Runtime/CableComponent/Source/CableComponent/Classes/CableComponent.h b/Engine/Plugins/Runtime/CableComponent/Source/CableComponent/Classes/CableComponent.h
--- a/Engine/Plugins/Runtime/CableComponent/Source/CableComponent/Classes/CableComponent.h
+++ b/Engine/Plugins/Runtime/CableComponent/Source/CableComponent/Classes/CableComponent.h
@@ -29,23 +29,23 @@ class CABLECOMPONENT_API UCableComponent : public UMeshComponent
 
 public:
 
-	// Begin UActorComponent interface.
+	//~ Begin UActorComponent Interface.
 	virtual void OnRegister() override;
 	virtual void TickComponent(float DeltaTime, enum ELevelTick TickType, FActorComponentTickFunction *ThisTickFunction) override;
 	virtual void SendRenderDynamicData_Concurrent() override;
-	// Begin UActorComponent interface.
+	//~ Begin UActorComponent Interface.
 
-	// Begin USceneComponent interface.
+	//~ Begin USceneComponent Interface.
 	virtual FBoxSphereBounds CalcBounds(const FTransform& LocalToWorld) const override;
-	// Begin USceneComponent interface.
+	//~ Begin USceneComponent Interface.
 
-	// Begin UPrimitiveComponent interface.
+	//~ Begin UPrimitiveComponent Interface.
 	virtual FPrimitiveSceneProxy* CreateSceneProxy() override;
-	// End UPrimitiveComponent interface.
+	//~ End UPrimitiveComponent Interface.
 
-	// Begin UMeshComponent interface.
+	//~ Begin UMeshComponent Interface.
 	virtual int32 GetNumMaterials() const override;
-	// End UMeshComponent interface.
+	//~ End UMeshComponent Interface.
 
 
 	/** Actor or Component that the end of the cable should be attached to */
diff --git a/Engine/Plugins/Runtime/CustomMeshComponent/Source/CustomMeshComponent/Classes/CustomMeshComponent.h b/Engine/Plugins/Runtime/CustomMeshComponent/Source/CustomMeshComponent/Classes/CustomMeshComponent.h
--- a/Engine/Plugins/Runtime/CustomMeshComponent/Source/CustomMeshComponent/Classes/CustomMeshComponent.h
+++ b/Engine/Plugins/Runtime/CustomMeshComponent/Source/CustomMeshComponent/Classes/CustomMeshComponent.h
@@ -39,17 +39,17 @@ class CUSTOMMESHCOMPONENT_API UCustomMeshComponent : public UMeshComponent
 
 private:
 
-	// Begin UPrimitiveComponent interface.
+	//~ Begin UPrimitiveComponent Interface.
 	virtual FPrimitiveSceneProxy* CreateSceneProxy() override;
-	// End UPrimitiveComponent interface.
+	//~ End UPrimitiveComponent Interface.
 
-	// Begin UMeshComponent interface.
+	//~ Begin UMeshComponent Interface.
 	virtual int32 GetNumMaterials() const override;
-	// End UMeshComponent interface.
+	//~ End UMeshComponent Interface.
 
-	// Begin USceneComponent interface.
+	//~ Begin USceneComponent Interface.
 	virtual FBoxSphereBounds CalcBounds(const FTransform& LocalToWorld) const override;
-	// Begin USceneComponent interface.
+	//~ Begin USceneComponent Interface.
 
 	/** */
 	TArray<FCustomMeshTriangle> CustomMeshTris;
diff --git a/Engine/Plugins/Runtime/ExampleDeviceProfileSelector/Source/ExampleDeviceProfileSelector/Private/ExampleDeviceProfileSelectorModule.h b/Engine/Plugins/Runtime/ExampleDeviceProfileSelector/Source/ExampleDeviceProfileSelector/Private/ExampleDeviceProfileSelectorModule.h
--- a/Engine/Plugins/Runtime/ExampleDeviceProfileSelector/Source/ExampleDeviceProfileSelector/Private/ExampleDeviceProfileSelectorModule.h
+++ b/Engine/Plugins/Runtime/ExampleDeviceProfileSelector/Source/ExampleDeviceProfileSelector/Private/ExampleDeviceProfileSelectorModule.h
@@ -11,15 +11,15 @@ class FExampleDeviceProfileSelectorModule
 {
 public:
 
-	// Begin IDeviceProfileSelectorModule interface
+	//~ Begin IDeviceProfileSelectorModule Interface
 	virtual const FString GetRuntimeDeviceProfileName() override;
-	// End IDeviceProfileSelectorModule interface
+	//~ End IDeviceProfileSelectorModule Interface
 
 
-	// Begin IModuleInterface interface
+	//~ Begin IModuleInterface Interface
 	virtual void StartupModule() override;
 	virtual void ShutdownModule() override;
-	// End IModuleInterface interface
+	//~ End IModuleInterface Interface
 
 	
 	/**
diff --git a/Engine/Plugins/Runtime/ProceduralMeshComponent/Source/Public/ProceduralMeshComponent.h b/Engine/Plugins/Runtime/ProceduralMeshComponent/Source/Public/ProceduralMeshComponent.h
--- a/Engine/Plugins/Runtime/ProceduralMeshComponent/Source/Public/ProceduralMeshComponent.h
+++ b/Engine/Plugins/Runtime/ProceduralMeshComponent/Source/Public/ProceduralMeshComponent.h
@@ -165,29 +165,29 @@ class PROCEDURALMESHCOMPONENT_API UProceduralMeshComponent : public UMeshCompone
 	UFUNCTION(BlueprintCallable, Category = "Components|ProceduralMesh")
 	bool IsMeshSectionVisible(int32 SectionIndex) const;
 
-	// Begin Interface_CollisionDataProvider Interface
+	//~ Begin Interface_CollisionDataProvider Interface
 	virtual bool GetPhysicsTriMeshData(struct FTriMeshCollisionData* CollisionData, bool InUseAllTriData) override;
 	virtual bool ContainsPhysicsTriMeshData(bool InUseAllTriData) const override;
 	virtual bool WantsNegXTriMesh() override{ return false; }
-	// End Interface_CollisionDataProvider Interface
+	//~ End Interface_CollisionDataProvider Interface
 
 	/** Collision data */
 	UPROPERTY(transient, duplicatetransient)
 	class UBodySetup* ProcMeshBodySetup;
 
 private:
-	// Begin USceneComponent interface.
+	//~ Begin USceneComponent Interface.
 	virtual FBoxSphereBounds CalcBounds(const FTransform& LocalToWorld) const override;
-	// Begin USceneComponent interface.
+	//~ Begin USceneComponent Interface.
 
-	// Begin UPrimitiveComponent interface.
+	//~ Begin UPrimitiveComponent Interface.
 	virtual FPrimitiveSceneProxy* CreateSceneProxy() override;
 	virtual class UBodySetup* GetBodySetup() override;
-	// End UPrimitiveComponent interface.
+	//~ End UPrimitiveComponent Interface.
 
-	// Begin UMeshComponent interface.
+	//~ Begin UMeshComponent Interface.
 	virtual int32 GetNumMaterials() const override;
-	// End UMeshComponent interface.
+	//~ End UMeshComponent Interface.
 
 
 
diff --git a/Engine/Plugins/Runtime/SoundMod/Source/SoundMod/Classes/SoundMod.h b/Engine/Plugins/Runtime/SoundMod/Source/SoundMod/Classes/SoundMod.h
--- a/Engine/Plugins/Runtime/SoundMod/Source/SoundMod/Classes/SoundMod.h
+++ b/Engine/Plugins/Runtime/SoundMod/Source/SoundMod/Classes/SoundMod.h
@@ -28,14 +28,14 @@ private:
 	uint8*	ResourceData;
 
 public:	
-	// Begin UObject interface. 
+	//~ Begin UObject Interface. 
 	virtual void Serialize(FArchive& Ar) override;
-	// End UObject interface. 
+	//~ End UObject Interface. 
 
-	// Begin USoundBase interface.
+	//~ Begin USoundBase Interface.
 	virtual bool IsPlayable() const override;
 	virtual void Parse(class FAudioDevice* AudioDevice, const UPTRINT NodeWaveInstanceHash, FActiveSound& ActiveSound, const FSoundParseParameters& ParseParams, TArray<FWaveInstance*>& WaveInstances) override;
 	virtual float GetMaxAudibleDistance() override;
-	// End USoundBase interface.
+	//~ End USoundBase Interface.
 };
 
diff --git a/Engine/Plugins/Runtime/SoundMod/Source/SoundMod/Classes/SoundModWave.h b/Engine/Plugins/Runtime/SoundMod/Source/SoundMod/Classes/SoundModWave.h
--- a/Engine/Plugins/Runtime/SoundMod/Source/SoundMod/Classes/SoundModWave.h
+++ b/Engine/Plugins/Runtime/SoundMod/Source/SoundMod/Classes/SoundModWave.h
@@ -19,8 +19,8 @@ class USoundModWave : public USoundWaveProcedural
 
 	xmp_context xmpContext;
 
-	// Begin USoundWave interface.
+	//~ Begin USoundWave Interface.
 	virtual int32 GeneratePCMData(uint8* PCMData, const int32 SamplesNeeded) override;
-	// End USoundWave interface.
+	//~ End USoundWave Interface.
 
 };
\ No newline at end of file
diff --git a/Engine/Plugins/Runtime/SoundMod/Source/SoundMod/Public/SoundNodeModPlayer.h b/Engine/Plugins/Runtime/SoundMod/Source/SoundMod/Public/SoundNodeModPlayer.h
--- a/Engine/Plugins/Runtime/SoundMod/Source/SoundMod/Public/SoundNodeModPlayer.h
+++ b/Engine/Plugins/Runtime/SoundMod/Source/SoundMod/Public/SoundNodeModPlayer.h
@@ -33,20 +33,20 @@ public:
 #if WITH_EDITOR
 	virtual void PostEditChangeProperty(FPropertyChangedEvent& PropertyChangedEvent) override;
 #endif
-	// End UObject Interface
+	//~ End UObject Interface
 
-	// Begin USoundNode Interface
+	//~ Begin USoundNode Interface
 	virtual int32 GetMaxChildNodes() const override;
 	virtual float GetDuration() override;
 	virtual void ParseNodes(FAudioDevice* AudioDevice, const UPTRINT NodeWaveInstanceHash, FActiveSound& ActiveSound, const FSoundParseParameters& ParseParams, TArray<FWaveInstance*>& WaveInstances) override;
 #if WITH_EDITOR
 	virtual FText GetTitle() const override;
 #endif
-	// End USoundNode Interface
+	//~ End USoundNode Interface
 
-	// Begin USoundNodeAssetReferencer Interface
+	//~ Begin USoundNodeAssetReferencer Interface
 	virtual void LoadAsset() override;
-	// End USoundNode Interface
+	//~ End USoundNode Interface
 
 };
 
diff --git a/Engine/Plugins/Runtime/WebBrowserWidget/Source/WebBrowserWidget/Public/WebBrowser.h b/Engine/Plugins/Runtime/WebBrowserWidget/Source/WebBrowserWidget/Public/WebBrowser.h
--- a/Engine/Plugins/Runtime/WebBrowserWidget/Source/WebBrowserWidget/Public/WebBrowser.h
+++ b/Engine/Plugins/Runtime/WebBrowserWidget/Source/WebBrowserWidget/Public/WebBrowser.h
@@ -41,9 +41,9 @@ public:
 
 public:
 
-	// UWidget interface
+	//~ Begin UWidget interface
 	virtual void SynchronizeProperties() override;
-	// End of UWidget interface
+	// End UWidget interface
 
 	virtual void ReleaseSlateResources(bool bReleaseChildren) override;
 
diff --git a/Engine/Plugins/ScriptPlugin/Source/ScriptEditorPlugin/Classes/ReimportScriptFactory.h b/Engine/Plugins/ScriptPlugin/Source/ScriptEditorPlugin/Classes/ReimportScriptFactory.h
--- a/Engine/Plugins/ScriptPlugin/Source/ScriptEditorPlugin/Classes/ReimportScriptFactory.h
+++ b/Engine/Plugins/ScriptPlugin/Source/ScriptEditorPlugin/Classes/ReimportScriptFactory.h
@@ -21,12 +21,12 @@ class UReimportScriptFactory : public UScriptFactory, public FReimportHandler
 
 	virtual bool ConfigureProperties() override;
 
-	// Begin FReimportHandler interface
+	//~ Begin FReimportHandler interface
 	virtual bool CanReimport(UObject* Obj, TArray<FString>& OutFilenames) override;
 	virtual void SetReimportPaths(UObject* Obj, const TArray<FString>& NewReimportPaths) override;
 	virtual EReimportResult::Type Reimport(UObject* Obj) override;
 	virtual int32 GetPriority() const override;
-	// End FReimportHandler interface
+	//~ End FReimportHandler interface
 };
 
 
diff --git a/Engine/Plugins/ScriptPlugin/Source/ScriptEditorPlugin/Classes/ScriptFactory.h b/Engine/Plugins/ScriptPlugin/Source/ScriptEditorPlugin/Classes/ScriptFactory.h
--- a/Engine/Plugins/ScriptPlugin/Source/ScriptEditorPlugin/Classes/ScriptFactory.h
+++ b/Engine/Plugins/ScriptPlugin/Source/ScriptEditorPlugin/Classes/ScriptFactory.h
@@ -18,9 +18,9 @@ public:
 	UPROPERTY(EditAnywhere, Category = ScriptFactory, meta = (AllowAbstract = "", BlueprintBaseOnly = ""))
 	TSubclassOf<class UObject> ParentClass;
 
-	// Begin UFactory Interface
+	//~ Begin UFactory Interface
 	virtual bool ConfigureProperties() override;
 	virtual bool DoesSupportClass(UClass* Class) override;
 	virtual UObject* FactoryCreateText(UClass* InClass, UObject* InParent, FName InName, EObjectFlags Flags, UObject* Context, const TCHAR* Type, const TCHAR*& Buffer, const TCHAR* BufferEnd, FFeedbackContext* Warn);
-	// End UFactory Interface
+	//~ End UFactory Interface
 };
diff --git a/Engine/Plugins/ScriptPlugin/Source/ScriptPlugin/Classes/ScriptContextComponent.h b/Engine/Plugins/ScriptPlugin/Source/ScriptPlugin/Classes/ScriptContextComponent.h
--- a/Engine/Plugins/ScriptPlugin/Source/ScriptPlugin/Classes/ScriptContextComponent.h
+++ b/Engine/Plugins/ScriptPlugin/Source/ScriptPlugin/Classes/ScriptContextComponent.h
@@ -14,12 +14,12 @@ class SCRIPTPLUGIN_API UScriptContextComponent : public UActorComponent
 
 public:
 
-	// Begin UActorComponent interface.
+	//~ Begin UActorComponent Interface.
 	virtual void OnRegister() override;
 	virtual void InitializeComponent() override;
 	virtual void TickComponent(float DeltaTime, enum ELevelTick TickType, FActorComponentTickFunction *ThisTickFunction) override;
 	virtual void BeginDestroy() override;
-	// Begin UActorComponent interface.
+	//~ Begin UActorComponent Interface.
 
 	/**
 	* Calls a script-defined function (no arguments)
diff --git a/Engine/Plugins/ScriptPlugin/Source/ScriptPlugin/Classes/ScriptPluginComponent.h b/Engine/Plugins/ScriptPlugin/Source/ScriptPlugin/Classes/ScriptPluginComponent.h
--- a/Engine/Plugins/ScriptPlugin/Source/ScriptPlugin/Classes/ScriptPluginComponent.h
+++ b/Engine/Plugins/ScriptPlugin/Source/ScriptPlugin/Classes/ScriptPluginComponent.h
@@ -21,12 +21,12 @@ public:
 	UFUNCTION(BlueprintCallable, Category = "Script|Functions")
 	virtual bool CallScriptFunction(FString FunctionName);
 
-	// Begin UActorComponent interface.
+	//~ Begin UActorComponent Interface.
 	virtual void OnRegister() override;
 	virtual void InitializeComponent() override;
 	virtual void TickComponent(float DeltaTime, enum ELevelTick TickType, FActorComponentTickFunction *ThisTickFunction) override;
 	virtual void OnUnregister() override;
-	// Begin UActorComponent interface.
+	//~ Begin UActorComponent Interface.
 
 protected:
 
diff --git a/Engine/Source/Developer/AllDesktopTargetPlatform/Private/AllDesktopTargetPlatform.h b/Engine/Source/Developer/AllDesktopTargetPlatform/Private/AllDesktopTargetPlatform.h
--- a/Engine/Source/Developer/AllDesktopTargetPlatform/Private/AllDesktopTargetPlatform.h
+++ b/Engine/Source/Developer/AllDesktopTargetPlatform/Private/AllDesktopTargetPlatform.h
@@ -52,18 +52,18 @@ public:
 
 public:
 
-	// Begin TTargetPlatformBase interface
+	//~ Begin TTargetPlatformBase Interface
 
 	virtual bool IsServerOnly( ) const override
 	{
 		return false;
 	}
 
-	// End TTargetPlatformBase interface
+	//~ End TTargetPlatformBase Interface
 
 public:
 
-	// Begin ITargetPlatform interface
+	//~ Begin ITargetPlatform Interface
 
 	virtual void EnableDeviceCheck(bool OnOff) override
 	{
diff --git a/Engine/Source/Developer/Android/AndroidTargetPlatform/Private/AndroidTargetDevice.h b/Engine/Source/Developer/Android/AndroidTargetPlatform/Private/AndroidTargetDevice.h
--- a/Engine/Source/Developer/Android/AndroidTargetPlatform/Private/AndroidTargetDevice.h
+++ b/Engine/Source/Developer/Android/AndroidTargetPlatform/Private/AndroidTargetDevice.h
@@ -99,7 +99,7 @@ public:
 
 public:
 
-	// Begin ITargetDevice interface
+	//~ Begin ITargetDevice Interface
 	virtual bool Connect() override
 	{
 		return true;
@@ -162,7 +162,7 @@ public:
 	virtual bool TerminateProcess(const int32 ProcessId) override;
 	virtual void SetUserCredentials(const FString& UserName, const FString& UserPassword) override;
 	virtual bool GetUserCredentials(FString& OutUserName, FString& OutUserPassword) override;
-	// End ITargetDevice interface
+	//~ End ITargetDevice Interface
 
 protected:
 
diff --git a/Engine/Source/Developer/Android/AndroidTargetPlatform/Private/AndroidTargetPlatform.h b/Engine/Source/Developer/Android/AndroidTargetPlatform/Private/AndroidTargetPlatform.h
--- a/Engine/Source/Developer/Android/AndroidTargetPlatform/Private/AndroidTargetPlatform.h
+++ b/Engine/Source/Developer/Android/AndroidTargetPlatform/Private/AndroidTargetPlatform.h
@@ -84,7 +84,7 @@ public:
 
 public:
 
-	// Begin ITargetPlatform interface
+	//~ Begin ITargetPlatform Interface
 
 	virtual void EnableDeviceCheck(bool OnOff) override {}
 
@@ -173,7 +173,7 @@ public:
 		return DeviceLostEvent;
 	}
 
-	// End ITargetPlatform interface
+	//~ End ITargetPlatform Interface
 
 protected:
 
diff --git a/Engine/Source/Developer/CollisionAnalyzer/Private/CollisionAnalyzer.h b/Engine/Source/Developer/CollisionAnalyzer/Private/CollisionAnalyzer.h
--- a/Engine/Source/Developer/CollisionAnalyzer/Private/CollisionAnalyzer.h
+++ b/Engine/Source/Developer/CollisionAnalyzer/Private/CollisionAnalyzer.h
@@ -26,7 +26,7 @@ struct FCAQuery
 class FCollisionAnalyzer : public ICollisionAnalyzer
 {
 public:
-	// Begin ICollisionAnalyzer interface
+	//~ Begin ICollisionAnalyzer Interface
 	virtual void CaptureQuery(
 		const FVector& Start, 
 		const FVector& End, 
@@ -47,7 +47,7 @@ public:
 
 	virtual void TickAnalyzer(UWorld* InWorld) override;
 	virtual bool IsRecording() override;
-	// End ICollisionAnalyzer interface
+	//~ End ICollisionAnalyzer Interface
 
 	/** Change the current recording state */
 	void SetIsRecording(bool bNewRecording);
diff --git a/Engine/Source/Developer/DerivedDataCache/Private/DDCCleanup.h b/Engine/Source/Developer/DerivedDataCache/Private/DDCCleanup.h
--- a/Engine/Source/Developer/DerivedDataCache/Private/DDCCleanup.h
+++ b/Engine/Source/Developer/DerivedDataCache/Private/DDCCleanup.h
@@ -25,12 +25,12 @@ class FDDCCleanup : public FRunnable
 	/** Destructor */
 	virtual ~FDDCCleanup();
 
-	// Begin FRunnable interface.
+	//~ Begin FRunnable Interface.
 	virtual bool Init();
 	virtual uint32 Run();
 	virtual void Stop();
 	virtual void Exit();
-	// End FRunnable interface
+	//~ End FRunnable Interface
 
 	/** Checks if there's been any Stop requests */
 	FORCEINLINE bool ShouldStop() const
diff --git a/Engine/Source/Developer/FunctionalTesting/Classes/FuncTestRenderingComponent.h b/Engine/Source/Developer/FunctionalTesting/Classes/FuncTestRenderingComponent.h
--- a/Engine/Source/Developer/FunctionalTesting/Classes/FuncTestRenderingComponent.h
+++ b/Engine/Source/Developer/FunctionalTesting/Classes/FuncTestRenderingComponent.h
@@ -9,12 +9,12 @@ class UFuncTestRenderingComponent : public UPrimitiveComponent
 {
 	GENERATED_UCLASS_BODY()
 
-	// Begin UPrimitiveComponent Interface
+	//~ Begin UPrimitiveComponent Interface
 	virtual FPrimitiveSceneProxy* CreateSceneProxy() override;
 	virtual bool ShouldRecreateProxyOnUpdateTransform() const override { return true; }
-	// End UPrimitiveComponent Interface
+	//~ End UPrimitiveComponent Interface
 
-	// Begin USceneComponent Interface
+	//~ Begin USceneComponent Interface
 	virtual FBoxSphereBounds CalcBounds(const FTransform &LocalToWorld) const override;
-	// End USceneComponent Interface
+	//~ End USceneComponent Interface
 };
\ No newline at end of file
diff --git a/Engine/Source/Developer/GameplayDebugger/Classes/GameplayDebuggerSettings.h b/Engine/Source/Developer/GameplayDebugger/Classes/GameplayDebuggerSettings.h
--- a/Engine/Source/Developer/GameplayDebugger/Classes/GameplayDebuggerSettings.h
+++ b/Engine/Source/Developer/GameplayDebugger/Classes/GameplayDebuggerSettings.h
@@ -37,10 +37,10 @@ class GAMEPLAYDEBUGGER_API UGameplayDebuggerSettings : public UObject
 	GENERATED_UCLASS_BODY()
 public:
 #if WITH_EDITOR
-	// Begin UObject Interface
+	//~ Begin UObject Interface
 	virtual void PostEditChangeProperty(struct FPropertyChangedEvent& PropertyChangedEvent) override;
 	virtual void PostInitProperties() override;
-	// End UObject Interface
+	//~ End UObject Interface
 #endif
 
 	DECLARE_EVENT_OneParam(UGameplayDebuggerSettings, FSettingChangedEvent, FName /*PropertyName*/);
diff --git a/Engine/Source/Developer/GameplayDebugger/Classes/GameplayDebuggingControllerComponent.h b/Engine/Source/Developer/GameplayDebugger/Classes/GameplayDebuggingControllerComponent.h
--- a/Engine/Source/Developer/GameplayDebugger/Classes/GameplayDebuggingControllerComponent.h
+++ b/Engine/Source/Developer/GameplayDebugger/Classes/GameplayDebuggingControllerComponent.h
@@ -112,9 +112,9 @@ class AGaneplayDebuggerProxyHUD : public AHUD
 
 	FFontRenderInfo TextRenderInfo;
 
-	// Begin AActor Interface
+	//~ Begin AActor Interface
 	virtual void PostRender() override;
-	// End AActor Interface
+	//~ End AActor Interface
 
 	TWeakObjectPtr<AHUD> RedirectedHUD;
 };
diff --git a/Engine/Source/Developer/HTML5/HTML5TargetPlatform/Private/HTML5TargetDevice.h b/Engine/Source/Developer/HTML5/HTML5TargetPlatform/Private/HTML5TargetDevice.h
--- a/Engine/Source/Developer/HTML5/HTML5TargetPlatform/Private/HTML5TargetDevice.h
+++ b/Engine/Source/Developer/HTML5/HTML5TargetPlatform/Private/HTML5TargetDevice.h
@@ -51,7 +51,7 @@ public:
 
 public:
 
-	// Begin ITargetDevice interface
+	//~ Begin ITargetDevice Interface
 
 	virtual bool Connect( ) override;
 
@@ -107,7 +107,7 @@ public:
 
 	virtual bool GetUserCredentials( FString& OutUserName, FString& OutUserPassword ) override;
 
-	// End ITargetDevice interface
+	//~ End ITargetDevice Interface
 
 private:
 
diff --git a/Engine/Source/Developer/HTML5/HTML5TargetPlatform/Private/HTML5TargetPlatform.h b/Engine/Source/Developer/HTML5/HTML5TargetPlatform/Private/HTML5TargetPlatform.h
--- a/Engine/Source/Developer/HTML5/HTML5TargetPlatform/Private/HTML5TargetPlatform.h
+++ b/Engine/Source/Developer/HTML5/HTML5TargetPlatform/Private/HTML5TargetPlatform.h
@@ -25,7 +25,7 @@ public:
 
 	void RefreshHTML5Setup();
 
-	// Begin ITargetPlatform interface
+	//~ Begin ITargetPlatform Interface
 
 	virtual void EnableDeviceCheck(bool OnOff) override {}
 
@@ -92,7 +92,7 @@ public:
 		return DeviceLostEvent;
 	}
 
-	// End ITargetPlatform interface
+	//~ End ITargetPlatform Interface
 
 private:
 
diff --git a/Engine/Source/Developer/ImageWrapper/Private/ExrImageWrapper.h b/Engine/Source/Developer/ImageWrapper/Private/ExrImageWrapper.h
--- a/Engine/Source/Developer/ImageWrapper/Private/ExrImageWrapper.h
+++ b/Engine/Source/Developer/ImageWrapper/Private/ExrImageWrapper.h
@@ -33,7 +33,7 @@ public:
 
 public:
 
-	// Begin FImageWrapper Interface
+	//~ Begin FImageWrapper Interface
 
 	virtual void Compress( int32 Quality ) override;
 
@@ -41,7 +41,7 @@ public:
 	
 	virtual bool SetCompressed( const void* InCompressedData, int32 InCompressedSize ) override;
 
-	// End FImageWrapper Interface
+	//~ End FImageWrapper Interface
 
 private:
 	template <Imf::PixelType OutputFormat, typename sourcetype>
diff --git a/Engine/Source/Developer/ImageWrapper/Private/PngImageWrapper.h b/Engine/Source/Developer/ImageWrapper/Private/PngImageWrapper.h
--- a/Engine/Source/Developer/ImageWrapper/Private/PngImageWrapper.h
+++ b/Engine/Source/Developer/ImageWrapper/Private/PngImageWrapper.h
@@ -31,7 +31,7 @@ public:
 
 public:
 
-	// Begin FImageWrapper Interface
+	//~ Begin FImageWrapper Interface
 
 	virtual void Compress( int32 Quality ) override;
 
@@ -41,7 +41,7 @@ public:
 
 	virtual void Uncompress( const ERGBFormat::Type InFormat, int32 InBitDepth ) override;
 	
-	// End FImageWrapper Interface
+	//~ End FImageWrapper Interface
 
 public:
 
diff --git a/Engine/Source/Developer/LauncherServices/Private/Launcher/Launcher.h b/Engine/Source/Developer/LauncherServices/Private/Launcher/Launcher.h
--- a/Engine/Source/Developer/LauncherServices/Private/Launcher/Launcher.h
+++ b/Engine/Source/Developer/LauncherServices/Private/Launcher/Launcher.h
@@ -11,9 +11,9 @@ class FLauncher
 {
 public:
 
-	// Begin ILauncher interface
+	//~ Begin ILauncher Interface
 
 	virtual ILauncherWorkerPtr Launch( const ITargetDeviceProxyManagerRef& DeviceProxyManager, const ILauncherProfileRef& Profile ) override;
 
-	// End ILauncher interface
+	//~ End ILauncher Interface
 };
diff --git a/Engine/Source/Developer/LauncherServices/Private/Launcher/LauncherTask.h b/Engine/Source/Developer/LauncherServices/Private/Launcher/LauncherTask.h
--- a/Engine/Source/Developer/LauncherServices/Private/Launcher/LauncherTask.h
+++ b/Engine/Source/Developer/LauncherServices/Private/Launcher/LauncherTask.h
@@ -132,7 +132,7 @@ public:
 
 public:
 
-	// Begin FRunnable interface
+	//~ Begin FRunnable Interface
 
 	virtual bool Init( ) override
 	{
@@ -183,11 +183,11 @@ public:
 
 	virtual void Exit( ) override { }
 
-	// End FRunnable interface
+	//~ End FRunnable Interface
 
 public:
 
-	// Begin ILauncherTask interface
+	//~ Begin ILauncherTask Interface
 
 	virtual void Cancel( ) override
 	{
@@ -250,7 +250,7 @@ public:
 		return TaskCompleted;
 	}
 
-	// End ILauncherTask interface
+	//~ End ILauncherTask Interface
 
 protected:
 
diff --git a/Engine/Source/Developer/LauncherServices/Private/Profiles/LauncherDeviceGroup.h b/Engine/Source/Developer/LauncherServices/Private/Profiles/LauncherDeviceGroup.h
--- a/Engine/Source/Developer/LauncherServices/Private/Profiles/LauncherDeviceGroup.h
+++ b/Engine/Source/Developer/LauncherServices/Private/Profiles/LauncherDeviceGroup.h
@@ -30,7 +30,7 @@ public:
 
 public:
 
-	// Begin ILauncherDeviceGroup interface
+	//~ Begin ILauncherDeviceGroup Interface
 
 	virtual void AddDevice(const FString& DeviceID) override
 	{
@@ -92,7 +92,7 @@ public:
 		Name = NewName;
 	}
 
-	// End ILauncherDeviceGroup interface
+	//~ End ILauncherDeviceGroup Interface
 
 private:
 
diff --git a/Engine/Source/Developer/LauncherServices/Private/Profiles/LauncherProfile.h b/Engine/Source/Developer/LauncherServices/Private/Profiles/LauncherProfile.h
--- a/Engine/Source/Developer/LauncherServices/Private/Profiles/LauncherProfile.h
+++ b/Engine/Source/Developer/LauncherServices/Private/Profiles/LauncherProfile.h
@@ -43,7 +43,7 @@ public:
 		SetDefaults();
 	}
 
-	// Begin ILauncherSimpleProfile interface
+	//~ Begin ILauncherSimpleProfile Interface
 
 	virtual const FString& GetDeviceName() const override
 	{
@@ -208,7 +208,7 @@ public:
 
 public:
 
-	// Begin ILauncherProfile interface
+	//~ Begin ILauncherProfile Interface
 
 	virtual void AddCookedCulture( const FString& CultureName ) override
 	{
@@ -1342,7 +1342,7 @@ public:
 		return EditorExe;
 	}
 
-	// End ILauncherProfile interface
+	//~ End ILauncherProfile Interface
 
 protected:
 
diff --git a/Engine/Source/Developer/LauncherServices/Private/Profiles/LauncherProfileLaunchRole.h b/Engine/Source/Developer/LauncherServices/Private/Profiles/LauncherProfileLaunchRole.h
--- a/Engine/Source/Developer/LauncherServices/Private/Profiles/LauncherProfileLaunchRole.h
+++ b/Engine/Source/Developer/LauncherServices/Private/Profiles/LauncherProfileLaunchRole.h
@@ -26,7 +26,7 @@ public:
 
 public:
 
-	// Begin ILauncherProfileLaunchRole interface
+	//~ Begin ILauncherProfileLaunchRole Interface
 
 	virtual const FString& GetAssignedDevice( ) const override
 	{
@@ -105,7 +105,7 @@ public:
 		VsyncEnabled = Enabled;
 	}
 
-	// End ILauncherProfileLaunchRole interface
+	//~ End ILauncherProfileLaunchRole Interface
 
 private:
 
diff --git a/Engine/Source/Developer/LauncherServices/Private/Profiles/LauncherProfileManager.h b/Engine/Source/Developer/LauncherServices/Private/Profiles/LauncherProfileManager.h
--- a/Engine/Source/Developer/LauncherServices/Private/Profiles/LauncherProfileManager.h
+++ b/Engine/Source/Developer/LauncherServices/Private/Profiles/LauncherProfileManager.h
@@ -23,7 +23,7 @@ public:
 	 */
 	void Load();
 
-	// Begin ILauncherProfileManager interface
+	//~ Begin ILauncherProfileManager Interface
 
 	virtual void AddDeviceGroup( const ILauncherDeviceGroupRef& DeviceGroup ) override;
 
@@ -97,7 +97,7 @@ public:
 
 	virtual void SetProjectPath(const FString& InProjectPath) override;
 
-	// End ILauncherProfileManager interface
+	//~ End ILauncherProfileManager Interface
 
 protected:
 
diff --git a/Engine/Source/Developer/LogVisualizer/Private/VisualLoggerHUD.h b/Engine/Source/Developer/LogVisualizer/Private/VisualLoggerHUD.h
--- a/Engine/Source/Developer/LogVisualizer/Private/VisualLoggerHUD.h
+++ b/Engine/Source/Developer/LogVisualizer/Private/VisualLoggerHUD.h
@@ -13,9 +13,9 @@ class AVisualLoggerHUD : public ADebugCameraHUD
 
 	virtual bool DisplayMaterials( float X, float& Y, float DY, UMeshComponent* MeshComp ) override;
 
-	// Begin AActor Interface
+	//~ Begin AActor Interface
 	virtual void PostRender() override;
-	// End AActor Interface
+	//~ End AActor Interface
 };
 
 
diff --git a/Engine/Source/Developer/ProfilerService/Private/ProfilerServiceManager.h b/Engine/Source/Developer/ProfilerService/Private/ProfilerServiceManager.h
--- a/Engine/Source/Developer/ProfilerService/Private/ProfilerServiceManager.h
+++ b/Engine/Source/Developer/ProfilerService/Private/ProfilerServiceManager.h
@@ -42,10 +42,10 @@ public:
 
 public:
 
-	// Begin IProfilerServiceManager interface
+	//~ Begin IProfilerServiceManager Interface
 	virtual void StartCapture() override;
 	virtual void StopCapture() override;
-	// End IProfilerServiceManager interface
+	//~ End IProfilerServiceManager Interface
 
 	/**
 	 * Creates a profiler service manager for shared use
diff --git a/Engine/Source/Developer/ScreenShotComparisonTools/Private/ScreenShotBaseNode.h b/Engine/Source/Developer/ScreenShotComparisonTools/Private/ScreenShotBaseNode.h
--- a/Engine/Source/Developer/ScreenShotComparisonTools/Private/ScreenShotBaseNode.h
+++ b/Engine/Source/Developer/ScreenShotComparisonTools/Private/ScreenShotBaseNode.h
@@ -28,7 +28,7 @@ public:
 
 public:
 
-	// Begin IScreenShotData interface
+	//~ Begin IScreenShotData Interface
 
 	virtual void AddScreenShotData( const FScreenShotDataItem& InScreenDataItem ) override;
 	virtual const FString& GetAssetName() const override;
@@ -39,7 +39,7 @@ public:
 	virtual bool SetFilter( TSharedPtr< ScreenShotFilterCollection > ScreenFilter ) override;
 	virtual void SetDisplayEveryNthScreenshot( int32 NewNthValue ) override;
 
-	// End IScreenShotData interface
+	//~ End IScreenShotData Interface
 
 protected:
 
diff --git a/Engine/Source/Developer/ScreenShotComparisonTools/Private/ScreenShotManager.h b/Engine/Source/Developer/ScreenShotComparisonTools/Private/ScreenShotManager.h
--- a/Engine/Source/Developer/ScreenShotComparisonTools/Private/ScreenShotManager.h
+++ b/Engine/Source/Developer/ScreenShotComparisonTools/Private/ScreenShotManager.h
@@ -31,7 +31,7 @@ public:
 
 public:
 
-	// Begin IScreenShotManager interface
+	//~ Begin IScreenShotManager Interface
 
 	virtual void GenerateLists() override;
 
@@ -45,7 +45,7 @@ public:
 
 	virtual void SetDisplayEveryNthScreenshot(int32 NewNth) override;
 
-	// End IScreenShotManager interface
+	//~ End IScreenShotManager Interface
 
 private:
 
diff --git a/Engine/Source/Developer/ScreenShotComparisonTools/Private/ScreenShotPlatformNode.h b/Engine/Source/Developer/ScreenShotComparisonTools/Private/ScreenShotPlatformNode.h
--- a/Engine/Source/Developer/ScreenShotComparisonTools/Private/ScreenShotPlatformNode.h
+++ b/Engine/Source/Developer/ScreenShotComparisonTools/Private/ScreenShotPlatformNode.h
@@ -23,7 +23,7 @@ public:
 
 public:
 
-	// Begin IScreenShotData interface
+	//~ Begin IScreenShotData Interface
 
 	virtual void AddScreenShotData( const FScreenShotDataItem& InScreenDataItem ) override 
 	{
@@ -77,7 +77,7 @@ public:
 	}
 
 
-	// End IScreenShotData interface
+	//~ End IScreenShotData Interface
 
 private:
 	// If > 1 we will only show every Nth screenshot.
diff --git a/Engine/Source/Developer/ScreenShotComparisonTools/Private/ScreenShotScreenNode.h b/Engine/Source/Developer/ScreenShotComparisonTools/Private/ScreenShotScreenNode.h
--- a/Engine/Source/Developer/ScreenShotComparisonTools/Private/ScreenShotScreenNode.h
+++ b/Engine/Source/Developer/ScreenShotComparisonTools/Private/ScreenShotScreenNode.h
@@ -22,7 +22,7 @@ public:
 
 public:
 
-	// Begin IScreenShotData interface
+	//~ Begin IScreenShotData Interface
 
 	virtual void AddScreenShotData( const FScreenShotDataItem& InScreenDataItem ) override 
 	{
@@ -57,9 +57,9 @@ public:
 		return FilteredChildren.Num() > 0;
 	}
 
-	// End IScreenShotData interface
+	//~ End IScreenShotData Interface
 
-	// Begin SScreenShotNodeBase interface
+	//~ Begin SScreenShotNodeBase Interface
 
 	virtual TSharedRef<IScreenShotData> CreateNode( const FString& ChildName ) override 
 	{
@@ -68,6 +68,6 @@ public:
 		return ScreenShotNode.ToSharedRef();
 	};
 
-	// End SScreenShotNodeBase interface
+	//~ End SScreenShotNodeBase Interface
 };
 
diff --git a/Engine/Source/Developer/iOS/IOSTargetPlatform/Private/IOSTargetPlatform.h b/Engine/Source/Developer/iOS/IOSTargetPlatform/Private/IOSTargetPlatform.h
--- a/Engine/Source/Developer/iOS/IOSTargetPlatform/Private/IOSTargetPlatform.h
+++ b/Engine/Source/Developer/iOS/IOSTargetPlatform/Private/IOSTargetPlatform.h
@@ -32,18 +32,18 @@ public:
 
 public:
 
-	// Begin TTargetPlatformBase interface
+	//~ Begin TTargetPlatformBase Interface
 
 	virtual bool IsServerOnly( ) const override
 	{
 		return false;
 	}
 
-	// End TTargetPlatformBase interface
+	//~ End TTargetPlatformBase Interface
 
 public:
 
-	// Begin ITargetPlatform interface
+	//~ Begin ITargetPlatform Interface
 
 	virtual void EnableDeviceCheck(bool OnOff) override;
 
@@ -126,7 +126,7 @@ public:
 		return DeviceLostEvent;
 	}
 
-	// Begin ITargetPlatform interface
+	//~ Begin ITargetPlatform Interface
 
 protected:
 
diff --git a/Engine/Source/Editor/AIGraph/Classes/AIGraphNode.h b/Engine/Source/Editor/AIGraph/Classes/AIGraphNode.h
--- a/Engine/Source/Editor/AIGraph/Classes/AIGraphNode.h
+++ b/Engine/Source/Editor/AIGraph/Classes/AIGraphNode.h
@@ -40,7 +40,7 @@ class AIGRAPH_API UAIGraphNode : public UEdGraphNode
 	UPROPERTY()
 	FString ErrorMessage;
 
-	// Begin UEdGraphNode Interface
+	//~ Begin UEdGraphNode Interface
 	virtual class UAIGraph* GetAIGraph();
 	virtual void AutowireNewNode(UEdGraphPin* FromPin) override;
 	virtual void PostPlacedNewNode() override;
@@ -52,9 +52,9 @@ class AIGRAPH_API UAIGraphNode : public UEdGraphNode
 	virtual void NodeConnectionListChanged() override;
 	virtual bool CanCreateUnderSpecifiedSchema(const UEdGraphSchema* DesiredSchema) const override;
 	virtual void FindDiffs(class UEdGraphNode* OtherNode, struct FDiffResults& Results) override;
-	// End UEdGraphNode Interface
+	//~ End UEdGraphNode Interface
 
-	// Begin UObject Interface
+	//~ Begin UObject Interface
 #if WITH_EDITOR
 	virtual void PostEditImport() override;
 	virtual void PostEditUndo() override;
diff --git a/Engine/Source/Editor/AIGraph/Classes/AIGraphSchema.h b/Engine/Source/Editor/AIGraph/Classes/AIGraphSchema.h
--- a/Engine/Source/Editor/AIGraph/Classes/AIGraphSchema.h
+++ b/Engine/Source/Editor/AIGraph/Classes/AIGraphSchema.h
@@ -25,11 +25,11 @@ struct AIGRAPH_API FAISchemaAction_NewNode : public FEdGraphSchemaAction
 		, NodeTemplate(nullptr)
 	{}
 
-	// FEdGraphSchemaAction interface
+	//~ Begin FEdGraphSchemaAction Interface
 	virtual UEdGraphNode* PerformAction(class UEdGraph* ParentGraph, UEdGraphPin* FromPin, const FVector2D Location, bool bSelectNewNode = true) override;
 	virtual UEdGraphNode* PerformAction(class UEdGraph* ParentGraph, TArray<UEdGraphPin*>& FromPins, const FVector2D Location, bool bSelectNewNode = true) override;
 	virtual void AddReferencedObjects(FReferenceCollector& Collector) override;
-	// End of FEdGraphSchemaAction interface
+	//~ End FEdGraphSchemaAction Interface
 
 	template <typename NodeType>
 	static NodeType* SpawnNodeFromTemplate(class UEdGraph* ParentGraph, NodeType* InTemplateNode, const FVector2D Location)
@@ -67,11 +67,11 @@ struct AIGRAPH_API FAISchemaAction_NewSubNode : public FEdGraphSchemaAction
 		, ParentNode(nullptr)
 	{}
 
-	// FEdGraphSchemaAction interface
+	//~ Begin FEdGraphSchemaAction Interface
 	virtual UEdGraphNode* PerformAction(class UEdGraph* ParentGraph, UEdGraphPin* FromPin, const FVector2D Location, bool bSelectNewNode = true) override;
 	virtual UEdGraphNode* PerformAction(class UEdGraph* ParentGraph, TArray<UEdGraphPin*>& FromPins, const FVector2D Location, bool bSelectNewNode = true) override;
 	virtual void AddReferencedObjects(FReferenceCollector& Collector) override;
-	// End of FEdGraphSchemaAction interface
+	//~ End FEdGraphSchemaAction Interface
 };
 
 UCLASS()
@@ -81,7 +81,7 @@ class AIGRAPH_API UAIGraphSchema : public UEdGraphSchema
 
 	void GetBreakLinkToSubMenuActions(class FMenuBuilder& MenuBuilder, class UEdGraphPin* InGraphPin);
 
-	// Begin EdGraphSchema interface
+	//~ Begin EdGraphSchema Interface
 	virtual void GetContextMenuActions(const UEdGraph* CurrentGraph, const UEdGraphNode* InGraphNode, const UEdGraphPin* InGraphPin, class FMenuBuilder* MenuBuilder, bool bIsDebugging) const override;
 	virtual FLinearColor GetPinTypeColor(const FEdGraphPinType& PinType) const override;
 	virtual bool ShouldHidePinDefaultValue(UEdGraphPin* Pin) const override;
@@ -89,7 +89,7 @@ class AIGRAPH_API UAIGraphSchema : public UEdGraphSchema
 	virtual void BreakNodeLinks(UEdGraphNode& TargetNode) const override;
 	virtual void BreakPinLinks(UEdGraphPin& TargetPin, bool bSendsNodeNotification) const override;
 	virtual void BreakSinglePinLink(UEdGraphPin* SourcePin, UEdGraphPin* TargetPin) override;
-	// End EdGraphSchema interface
+	//~ End EdGraphSchema Interface
 
 	virtual void GetGraphNodeContextActions(FGraphContextMenuBuilder& ContextMenuBuilder, int32 SubNodeFlags) const;
 	virtual void GetSubNodeClasses(int32 SubNodeFlags, TArray<FGraphNodeClassData>& ClassData, UClass*& GraphNodeClass) const;
diff --git a/Engine/Source/Editor/AIGraph/Public/AIGraphEditor.h b/Engine/Source/Editor/AIGraph/Public/AIGraphEditor.h
--- a/Engine/Source/Editor/AIGraph/Public/AIGraphEditor.h
+++ b/Engine/Source/Editor/AIGraph/Public/AIGraphEditor.h
@@ -13,7 +13,7 @@ public:
 	FGraphPanelSelectionSet GetSelectedNodes() const;
 	virtual void OnSelectedNodesChanged(const TSet<class UObject*>& NewSelection);
 
-	// Begin FEditorUndoClient Interface
+	//~ Begin FEditorUndoClient Interface
 	virtual void PostUndo(bool bSuccess) override;
 	virtual void PostRedo(bool bSuccess) override;
 	// End of FEditorUndoClient
diff --git a/Engine/Source/Editor/AIGraph/Public/SGraphNodeAI.h b/Engine/Source/Editor/AIGraph/Public/SGraphNodeAI.h
--- a/Engine/Source/Editor/AIGraph/Public/SGraphNodeAI.h
+++ b/Engine/Source/Editor/AIGraph/Public/SGraphNodeAI.h
@@ -30,7 +30,7 @@ public:
 
 	void Construct(const FArguments& InArgs, UAIGraphNode* InNode);
 
-	// SGraphNode interface
+	//~ Begin SGraphNode Interface
 	virtual TSharedPtr<SToolTip> GetComplexTooltip() override;
 	virtual void OnDragEnter(const FGeometry& MyGeometry, const FDragDropEvent& DragDropEvent) override;
 	virtual FReply OnDragOver(const FGeometry& MyGeometry, const FDragDropEvent& DragDropEvent) override;
@@ -39,7 +39,7 @@ public:
 	virtual FReply OnMouseMove(const FGeometry& SenderGeometry, const FPointerEvent& MouseEvent) override;
 	virtual void SetOwner(const TSharedRef<SGraphPanel>& OwnerPanel) override;
 	virtual void AddPin(const TSharedRef<SGraphPin>& PinToAdd) override;
-	// End of SGraphNode interface
+	//~ End SGraphNode Interface
 
 	/** handle mouse down on the node */
 	FReply OnMouseDown(const FGeometry& SenderGeometry, const FPointerEvent& MouseEvent);
@@ -76,10 +76,10 @@ public:
 
 	void Construct(const FArguments& InArgs, UEdGraphPin* InPin);
 protected:
-	// Begin SGraphPin interface
+	//~ Begin SGraphPin Interface
 	virtual FSlateColor GetPinColor() const override;
 	virtual TSharedRef<SWidget>	GetDefaultValueWidget() override;
-	// End SGraphPin interface
+	//~ End SGraphPin Interface
 
 	const FSlateBrush* GetPinBorder() const;
 };
diff --git a/Engine/Source/Editor/AnimGraph/Classes/AnimGraphNode_Root.h b/Engine/Source/Editor/AnimGraph/Classes/AnimGraphNode_Root.h
--- a/Engine/Source/Editor/AnimGraph/Classes/AnimGraphNode_Root.h
+++ b/Engine/Source/Editor/AnimGraph/Classes/AnimGraphNode_Root.h
@@ -13,20 +13,20 @@ class UAnimGraphNode_Root : public UAnimGraphNode_Base
 	UPROPERTY(EditAnywhere, Category=Settings)
 	FAnimNode_Root Node;
 
-	// Begin UEdGraphNode interface.
+	//~ Begin UEdGraphNode Interface.
 	virtual FLinearColor GetNodeTitleColor() const override;
 	virtual FText GetTooltipText() const override;
 	virtual FText GetNodeTitle(ENodeTitleType::Type TitleType) const override;
 	virtual bool CanUserDeleteNode() const override { return false; }
 	virtual bool CanDuplicateNode() const override { return false; }
 	virtual void GetMenuActions(FBlueprintActionDatabaseRegistrar& ActionRegistrar) const override;
-	// End UEdGraphNode interface.
+	//~ End UEdGraphNode Interface.
 
-	// UAnimGraphNode_Base interface
+	//~ Begin UAnimGraphNode_Base Interface
 	virtual bool IsSinkNode() const override;
 
 	// Get the link to the documentation
 	virtual FString GetDocumentationLink() const override;
 
-	// End of UAnimGraphNode_Base interface
+	//~ End UAnimGraphNode_Base Interface
 };
diff --git a/Engine/Source/Editor/AnimGraph/Classes/AnimGraphNode_RotateRootBone.h b/Engine/Source/Editor/AnimGraph/Classes/AnimGraphNode_RotateRootBone.h
--- a/Engine/Source/Editor/AnimGraph/Classes/AnimGraphNode_RotateRootBone.h
+++ b/Engine/Source/Editor/AnimGraph/Classes/AnimGraphNode_RotateRootBone.h
@@ -14,13 +14,13 @@ class UAnimGraphNode_RotateRootBone : public UAnimGraphNode_Base
 	UPROPERTY(EditAnywhere, Category=Settings)
 	FAnimNode_RotateRootBone Node;
 
-	// Begin UEdGraphNode interface.
+	//~ Begin UEdGraphNode Interface.
 	virtual FLinearColor GetNodeTitleColor() const override;
 	virtual FText GetTooltipText() const override;
 	virtual FText GetNodeTitle(ENodeTitleType::Type TitleType) const override;
-	// End UEdGraphNode interface.
+	//~ End UEdGraphNode Interface.
 
-	// UAnimGraphNode_Base interface
+	//~ Begin UAnimGraphNode_Base Interface
 	virtual FString GetNodeCategory() const override;
-	// End of UAnimGraphNode_Base interface
+	//~ End UAnimGraphNode_Base Interface
 };
diff --git a/Engine/Source/Editor/AnimGraph/Classes/AnimGraphNode_SpringBone.h b/Engine/Source/Editor/AnimGraph/Classes/AnimGraphNode_SpringBone.h
--- a/Engine/Source/Editor/AnimGraph/Classes/AnimGraphNode_SpringBone.h
+++ b/Engine/Source/Editor/AnimGraph/Classes/AnimGraphNode_SpringBone.h
@@ -14,20 +14,20 @@ class UAnimGraphNode_SpringBone : public UAnimGraphNode_SkeletalControlBase
 	UPROPERTY(EditAnywhere, Category=Settings)
 	FAnimNode_SpringBone Node;
 
-	// Begin UObject interface
+	//~ Begin UObject Interface
 	virtual void PostLoad();
-	// End UObject interface
+	//~ End UObject Interface
 
 public:
-	// UEdGraphNode interface
+	//~ Begin UEdGraphNode Interface
 	virtual FText GetNodeTitle(ENodeTitleType::Type TitleType) const override;
 	virtual FText GetTooltipText() const override;
-	// End of UEdGraphNode interface
+	//~ End UEdGraphNode Interface
 
 protected:
-	// UAnimGraphNode_SkeletalControlBase interface
+	//~ Begin UAnimGraphNode_SkeletalControlBase Interface
 	virtual FText GetControllerDescription() const override;
-	// End of UAnimGraphNode_SkeletalControlBase interface
+	//~ End UAnimGraphNode_SkeletalControlBase Interface
 
 private:
 	/** Constructing FText strings can be costly, so we cache the node's title */
diff --git a/Engine/Source/Editor/AnimGraph/Classes/AnimPreviewInstance.h b/Engine/Source/Editor/AnimGraph/Classes/AnimPreviewInstance.h
--- a/Engine/Source/Editor/AnimGraph/Classes/AnimPreviewInstance.h
+++ b/Engine/Source/Editor/AnimGraph/Classes/AnimPreviewInstance.h
@@ -48,21 +48,21 @@ class ANIMGRAPH_API UAnimPreviewInstance : public UAnimSingleNodeInstance
 	UPROPERTY(transient)
 	int32 MontagePreviewStartSectionIdx;
 
-	// Begin UAnimInstance interface
+	//~ Begin UAnimInstance Interface
 	virtual void NativeInitializeAnimation() override;
 	virtual void UpdateAnimationNode(float DeltaTimeX) override;
 	virtual bool NativeEvaluateAnimation(FPoseContext& Output) override;
-	// End UAnimInstance interface
+	//~ End UAnimInstance Interface
 
 	/** Set SkeletalControl Alpha**/
 	void SetSkeletalControlAlpha(float SkeletalControlAlpha);
 
 	UAnimSequence* GetAnimSequence();
 
-	// Begin UAnimSingleNodeInstance interface
+	//~ Begin UAnimSingleNodeInstance Interface
 	virtual void RestartMontage(UAnimMontage* Montage, FName FromSection = FName()) override;
 	virtual void SetAnimationAsset(UAnimationAsset* NewAsset, bool bIsLooping = true, float InPlayRate = 1.f) override;
-	// End UAnimSingleNodeInstance interface
+	//~ End UAnimSingleNodeInstance Interface
 
 	/** Montage preview functions */
 	void MontagePreview_JumpToStart();
diff --git a/Engine/Source/Editor/AnimGraph/Classes/AnimStateConduitNode.h b/Engine/Source/Editor/AnimGraph/Classes/AnimStateConduitNode.h
--- a/Engine/Source/Editor/AnimGraph/Classes/AnimStateConduitNode.h
+++ b/Engine/Source/Editor/AnimGraph/Classes/AnimStateConduitNode.h
@@ -14,7 +14,7 @@ public:
 	UPROPERTY()
 	class UEdGraph* BoundGraph;
 
-	// Begin UEdGraphNode interface
+	//~ Begin UEdGraphNode Interface
 	virtual void AllocateDefaultPins() override;
 	virtual void AutowireNewNode(UEdGraphPin* FromPin) override;
 	virtual FText GetNodeTitle(ENodeTitleType::Type TitleType) const override;
@@ -24,14 +24,14 @@ public:
 	virtual void PostPlacedNewNode() override;
 	virtual void DestroyNode() override;
 	virtual void ValidateNodeDuringCompilation(class FCompilerResultsLog& MessageLog) const override;
-	// End UEdGraphNode interface
+	//~ End UEdGraphNode Interface
 
-	// Begin UAnimStateNodeBase interface
+	//~ Begin UAnimStateNodeBase Interface
 	virtual UEdGraphPin* GetInputPin() const override;
 	virtual UEdGraphPin* GetOutputPin() const override;
 	virtual FString GetStateName() const override;
 	virtual void GetTransitionList(TArray<class UAnimStateTransitionNode*>& OutTransitions, bool bWantSortedList = true) override;
 	virtual FString GetDesiredNewNodeName() const override;
 	virtual UEdGraph* GetBoundGraph() const override { return BoundGraph; }
-	// End of UAnimStateNodeBase interface
+	//~ End UAnimStateNodeBase Interface
 };
diff --git a/Engine/Source/Editor/AnimGraph/Classes/AnimStateEntryNode.h b/Engine/Source/Editor/AnimGraph/Classes/AnimStateEntryNode.h
--- a/Engine/Source/Editor/AnimGraph/Classes/AnimStateEntryNode.h
+++ b/Engine/Source/Editor/AnimGraph/Classes/AnimStateEntryNode.h
@@ -10,11 +10,11 @@ class UAnimStateEntryNode : public UEdGraphNode
 	GENERATED_UCLASS_BODY()
 
 
-	// Begin UEdGraphNode interface
+	//~ Begin UEdGraphNode Interface
 	virtual void AllocateDefaultPins() override;
 	virtual FText GetNodeTitle(ENodeTitleType::Type TitleType) const override;
 	virtual FText GetTooltipText() const override;
-	// End UEdGraphNode interface
+	//~ End UEdGraphNode Interface
 	
 	ANIMGRAPH_API UEdGraphNode* GetOutputNode() const;
 
diff --git a/Engine/Source/Editor/AnimGraph/Classes/AnimStateNode.h b/Engine/Source/Editor/AnimGraph/Classes/AnimStateNode.h
--- a/Engine/Source/Editor/AnimGraph/Classes/AnimStateNode.h
+++ b/Engine/Source/Editor/AnimGraph/Classes/AnimStateNode.h
@@ -35,11 +35,11 @@ public:
 	UPROPERTY(EditAnywhere, Category=Events)
 	FAnimNotifyEvent StateFullyBlended;
 
-	// Begin UObject interface
+	//~ Begin UObject Interface
 	virtual void PostEditChangeProperty(struct FPropertyChangedEvent& PropertyChangedEvent) override;
-	// End UObject interface
+	//~ End UObject Interface
 
-	// Begin UEdGraphNode interface
+	//~ Begin UEdGraphNode Interface
 	virtual void AllocateDefaultPins() override;
 	virtual void AutowireNewNode(UEdGraphPin* FromPin) override;
 	virtual FText GetNodeTitle(ENodeTitleType::Type TitleType) const override;
@@ -48,14 +48,14 @@ public:
 	virtual void PostPasteNode() override;
 	virtual void PostPlacedNewNode() override;
 	virtual void DestroyNode() override;
-	// End UEdGraphNode interface
+	//~ End UEdGraphNode Interface
 	
-	// Begin UAnimStateNodeBase interface
+	//~ Begin UAnimStateNodeBase Interface
 	virtual UEdGraphPin* GetInputPin() const override;
 	virtual UEdGraphPin* GetOutputPin() const override;
 	virtual FString GetStateName() const override;
 	virtual void GetTransitionList(TArray<class UAnimStateTransitionNode*>& OutTransitions, bool bWantSortedList = false) override;
-	// End of UAnimStateNodeBase interface
+	//~ End UAnimStateNodeBase Interface
 
 	// @return the pose pin of the state sink node within the anim graph of this state
 	ANIMGRAPH_API UEdGraphPin* GetPoseSinkPinInsideState() const;
diff --git a/Engine/Source/Editor/AnimGraph/Classes/AnimStateTransitionNode.h b/Engine/Source/Editor/AnimGraph/Classes/AnimStateTransitionNode.h
--- a/Engine/Source/Editor/AnimGraph/Classes/AnimStateTransitionNode.h
+++ b/Engine/Source/Editor/AnimGraph/Classes/AnimStateTransitionNode.h
@@ -94,12 +94,12 @@ class UAnimStateTransitionNode : public UAnimStateNodeBase
 	UPROPERTY()
 	int32 SharedCrossfadeIdx;
 
-	// Begin UObject interface
+	//~ Begin UObject Interface
 	virtual void PostEditChangeProperty(struct FPropertyChangedEvent& PropertyChangedEvent) override;
 	virtual void PostLoad() override;
-	// End UObject interface
+	//~ End UObject Interface
 
-	// Begin UEdGraphNode interface
+	//~ Begin UEdGraphNode Interface
 	virtual void AllocateDefaultPins() override;
 	virtual FText GetNodeTitle(ENodeTitleType::Type TitleType) const override;
 	virtual FText GetTooltipText() const override;
@@ -111,13 +111,13 @@ class UAnimStateTransitionNode : public UAnimStateNodeBase
 	virtual void PostPlacedNewNode() override;
 	virtual void DestroyNode() override;
 	virtual void ValidateNodeDuringCompilation(class FCompilerResultsLog& MessageLog) const override;
-	// End UEdGraphNode interface
+	//~ End UEdGraphNode Interface
 
-	// Begin UAnimStateNodeBase interface
+	//~ Begin UAnimStateNodeBase Interface
 	virtual UEdGraph* GetBoundGraph() const override { return BoundGraph; }
 	virtual UEdGraphPin* GetInputPin() const override { return Pins[0]; }
 	virtual UEdGraphPin* GetOutputPin() const override { return Pins[1]; }
-	// End UAnimStateNodeBase interface
+	//~ End UAnimStateNodeBase Interface
 
 	// @return the name of this state
 	ANIMGRAPH_API FString GetStateName() const override;
diff --git a/Engine/Source/Editor/AnimGraph/Classes/AnimationCustomTransitionSchema.h b/Engine/Source/Editor/AnimGraph/Classes/AnimationCustomTransitionSchema.h
--- a/Engine/Source/Editor/AnimGraph/Classes/AnimationCustomTransitionSchema.h
+++ b/Engine/Source/Editor/AnimGraph/Classes/AnimationCustomTransitionSchema.h
@@ -9,8 +9,8 @@ class UAnimationCustomTransitionSchema : public UAnimationGraphSchema
 {
 	GENERATED_UCLASS_BODY()
 
-	// Begin UEdGraphSchema interface.
+	//~ Begin UEdGraphSchema Interface.
 	virtual void CreateDefaultNodesForGraph(UEdGraph& Graph) const override;
 	void GetGraphDisplayInformation(const UEdGraph& Graph, /*out*/ FGraphDisplayInfo& DisplayInfo) const override;
-	// End UEdGraphSchema interface.
+	//~ End UEdGraphSchema Interface.
 };
diff --git a/Engine/Source/Editor/AnimGraph/Classes/AnimationGraphSchema.h b/Engine/Source/Editor/AnimGraph/Classes/AnimationGraphSchema.h
--- a/Engine/Source/Editor/AnimGraph/Classes/AnimationGraphSchema.h
+++ b/Engine/Source/Editor/AnimGraph/Classes/AnimationGraphSchema.h
@@ -35,7 +35,7 @@ class UAnimationGraphSchema : public UEdGraphSchema_K2
 	UPROPERTY()
 	FName DefaultEvaluationHandlerName;
 
-	// Begin UEdGraphSchema interface.
+	//~ Begin UEdGraphSchema Interface.
 	virtual FLinearColor GetPinTypeColor(const FEdGraphPinType& PinType) const override;
 	virtual EGraphType GetGraphType(const UEdGraph* TestEdGraph) const override;
 	virtual void CreateDefaultNodesForGraph(UEdGraph& Graph) const override;
@@ -52,14 +52,14 @@ class UAnimationGraphSchema : public UEdGraphSchema_K2
 	virtual bool CanDuplicateGraph(UEdGraph* InSourceGraph) const override {	return false; }
 	virtual bool DoesSupportEventDispatcher() const	override { return false; }
 	virtual bool ShouldAlwaysPurgeOnModification() const override { return true; }
-	// End UEdGraphSchema interface.
+	//~ End UEdGraphSchema Interface.
 
-	// Begin UEdGraphSchema_K2 interface
+	//~ Begin UEdGraphSchema_K2 Interface
 	virtual const FPinConnectionResponse DetermineConnectionResponseOfCompatibleTypedPins(const UEdGraphPin* PinA, const UEdGraphPin* PinB, const UEdGraphPin* InputPin, const UEdGraphPin* OutputPin) const override;
 	virtual bool SearchForAutocastFunction(const UEdGraphPin* OutputPin, const UEdGraphPin* InputPin, /*out*/ FName& TargetFunction, /*out*/ UClass*& FunctionOwner) const override;
 	virtual bool ArePinsCompatible(const UEdGraphPin* PinA, const UEdGraphPin* PinB, const UClass* CallingContext = NULL, bool bIgnoreArray = false) const override;
 	virtual bool DoesSupportAnimNotifyActions() const override;
-	// End UEdGraphSchema_K2 interface
+	//~ End UEdGraphSchema_K2 Interface
 
 	/** Spawn the correct node in the Animation Graph using the given AnimationAsset at the supplied location */
 	static void SpawnNodeFromAsset(UAnimationAsset* Asset, const FVector2D& GraphPosition, UEdGraph* Graph, UEdGraphPin* PinIfAvailable);
diff --git a/Engine/Source/Editor/AnimGraph/Classes/AnimationStateGraphSchema.h b/Engine/Source/Editor/AnimGraph/Classes/AnimationStateGraphSchema.h
--- a/Engine/Source/Editor/AnimGraph/Classes/AnimationStateGraphSchema.h
+++ b/Engine/Source/Editor/AnimGraph/Classes/AnimationStateGraphSchema.h
@@ -9,8 +9,8 @@ class UAnimationStateGraphSchema : public UAnimationGraphSchema
 {
 	GENERATED_UCLASS_BODY()
 
-	// Begin UEdGraphSchema interface.
+	//~ Begin UEdGraphSchema Interface.
 	virtual void CreateDefaultNodesForGraph(UEdGraph& Graph) const override;
 	virtual void GetGraphDisplayInformation(const UEdGraph& Graph, /*out*/ FGraphDisplayInfo& DisplayInfo) const override;
-	// End UEdGraphSchema interface.
+	//~ End UEdGraphSchema Interface.
 };
diff --git a/Engine/Source/Editor/AnimGraph/Classes/AnimationStateMachineSchema.h b/Engine/Source/Editor/AnimGraph/Classes/AnimationStateMachineSchema.h
--- a/Engine/Source/Editor/AnimGraph/Classes/AnimationStateMachineSchema.h
+++ b/Engine/Source/Editor/AnimGraph/Classes/AnimationStateMachineSchema.h
@@ -49,9 +49,9 @@ struct ANIMGRAPH_API FEdGraphSchemaAction_NewStateComment : public FEdGraphSchem
 		: FEdGraphSchemaAction(InNodeCategory, InMenuDesc, InToolTip, InGrouping)
 	{}
 
-	// FEdGraphSchemaAction interface
+	//~ Begin FEdGraphSchemaAction Interface
 	virtual UEdGraphNode* PerformAction(class UEdGraph* ParentGraph, UEdGraphPin* FromPin, const FVector2D Location, bool bSelectNewNode = true) override;
-	// End of FEdGraphSchemaAction interface
+	//~ End FEdGraphSchemaAction Interface
 };
 
 
@@ -63,7 +63,7 @@ class UAnimationStateMachineSchema : public UEdGraphSchema
 	UPROPERTY()
 	FString PC_Exec;
 
-	// Begin UEdGraphSchema interface
+	//~ Begin UEdGraphSchema Interface
 	virtual void CreateDefaultNodesForGraph(UEdGraph& Graph) const override;
 	virtual const FPinConnectionResponse CanCreateConnection(const UEdGraphPin* A, const UEdGraphPin* B) const override;
 	virtual bool TryCreateConnection(UEdGraphPin* A, UEdGraphPin* B) const override;
@@ -78,5 +78,5 @@ class UAnimationStateMachineSchema : public UEdGraphSchema
 	virtual void DroppedAssetsOnPin(const TArray<FAssetData>& Assets, const FVector2D& GraphPosition, UEdGraphPin* Pin) const override;
 	virtual void GetAssetsNodeHoverMessage(const TArray<FAssetData>& Assets, const UEdGraphNode* HoverNode, FString& OutTooltipText, bool& OutOkIcon) const override;
 	virtual void GetAssetsPinHoverMessage(const TArray<FAssetData>& Assets, const UEdGraphPin* HoverPin, FString& OutTooltipText, bool& OutOkIcon) const override;
-	// End UEdGraphSchema interface
+	//~ End UEdGraphSchema Interface
 };
diff --git a/Engine/Source/Editor/BehaviorTreeEditor/Classes/BehaviorTreeGraphNode.h b/Engine/Source/Editor/BehaviorTreeEditor/Classes/BehaviorTreeGraphNode.h
--- a/Engine/Source/Editor/BehaviorTreeEditor/Classes/BehaviorTreeGraphNode.h
+++ b/Engine/Source/Editor/BehaviorTreeEditor/Classes/BehaviorTreeGraphNode.h
@@ -19,13 +19,13 @@ class UBehaviorTreeGraphNode : public UAIGraphNode
 	UPROPERTY()
 	TArray<UBehaviorTreeGraphNode*> Services;
 
-	// Begin UEdGraphNode Interface
+	//~ Begin UEdGraphNode Interface
 	virtual class UBehaviorTreeGraph* GetBehaviorTreeGraph();
 	virtual void AllocateDefaultPins() override;
 	virtual FText GetTooltipText() const override;
 	virtual bool CanCreateUnderSpecifiedSchema(const UEdGraphSchema* DesiredSchema) const override;
 	virtual void FindDiffs(class UEdGraphNode* OtherNode, struct FDiffResults& Results) override;
-	// End UEdGraphNode Interface
+	//~ End UEdGraphNode Interface
 
 	virtual FText GetDescription() const override;
 	virtual bool HasErrors() const override;
diff --git a/Engine/Source/Editor/BehaviorTreeEditor/Classes/EdGraphSchema_BehaviorTree.h b/Engine/Source/Editor/BehaviorTreeEditor/Classes/EdGraphSchema_BehaviorTree.h
--- a/Engine/Source/Editor/BehaviorTreeEditor/Classes/EdGraphSchema_BehaviorTree.h
+++ b/Engine/Source/Editor/BehaviorTreeEditor/Classes/EdGraphSchema_BehaviorTree.h
@@ -18,9 +18,9 @@ struct FBehaviorTreeSchemaAction_AutoArrange : public FEdGraphSchemaAction
 		: FEdGraphSchemaAction(InNodeCategory, InMenuDesc, InToolTip, InGrouping) 
 	{}
 
-	// FEdGraphSchemaAction interface
+	//~ Begin FEdGraphSchemaAction Interface
 	virtual UEdGraphNode* PerformAction(class UEdGraph* ParentGraph, UEdGraphPin* FromPin, const FVector2D Location, bool bSelectNewNode = true) override;
-	// End of FEdGraphSchemaAction interface
+	//~ End FEdGraphSchemaAction Interface
 };
 
 UCLASS(MinimalAPI)
@@ -28,7 +28,7 @@ class UEdGraphSchema_BehaviorTree : public UAIGraphSchema
 {
 	GENERATED_UCLASS_BODY()
 
-	// Begin EdGraphSchema interface
+	//~ Begin EdGraphSchema Interface
 	virtual void CreateDefaultNodesForGraph(UEdGraph& Graph) const override;
 	virtual void GetGraphContextActions(FGraphContextMenuBuilder& ContextMenuBuilder) const override;
 	virtual void GetContextMenuActions(const UEdGraph* CurrentGraph, const UEdGraphNode* InGraphNode, const UEdGraphPin* InGraphPin, class FMenuBuilder* MenuBuilder, bool bIsDebugging) const override;
@@ -40,7 +40,7 @@ class UEdGraphSchema_BehaviorTree : public UAIGraphSchema
 	virtual bool IsCacheVisualizationOutOfDate(int32 InVisualizationCacheID) const override;
 	virtual int32 GetCurrentVisualizationCacheID() const override;
 	virtual void ForceVisualizationCacheClear() const override;
-	// End EdGraphSchema interface
+	//~ End EdGraphSchema Interface
 
 	virtual void GetGraphNodeContextActions(FGraphContextMenuBuilder& ContextMenuBuilder, int32 SubNodeFlags) const override;
 	virtual void GetSubNodeClasses(int32 SubNodeFlags, TArray<FGraphNodeClassData>& ClassData, UClass*& GraphNodeClass) const override;
diff --git a/Engine/Source/Editor/BehaviorTreeEditor/Classes/EdGraphSchema_BehaviorTreeDecorator.h b/Engine/Source/Editor/BehaviorTreeEditor/Classes/EdGraphSchema_BehaviorTreeDecorator.h
--- a/Engine/Source/Editor/BehaviorTreeEditor/Classes/EdGraphSchema_BehaviorTreeDecorator.h
+++ b/Engine/Source/Editor/BehaviorTreeEditor/Classes/EdGraphSchema_BehaviorTreeDecorator.h
@@ -23,11 +23,11 @@ struct FDecoratorSchemaAction_NewNode : public FEdGraphSchemaAction
 		, NodeTemplate(NULL)
 	{}
 
-	// FEdGraphSchemaAction interface
+	//~ Begin FEdGraphSchemaAction Interface
 	virtual UEdGraphNode* PerformAction(class UEdGraph* ParentGraph, UEdGraphPin* FromPin, const FVector2D Location, bool bSelectNewNode = true) override;
 	virtual UEdGraphNode* PerformAction(class UEdGraph* ParentGraph, TArray<UEdGraphPin*>& FromPins, const FVector2D Location, bool bSelectNewNode = true) override;
 	virtual void AddReferencedObjects( FReferenceCollector& Collector ) override;
-	// End of FEdGraphSchemaAction interface
+	//~ End FEdGraphSchemaAction Interface
 
 	template <typename NodeType>
 	static NodeType* SpawnNodeFromTemplate(class UEdGraph* ParentGraph, NodeType* InTemplateNode, const FVector2D Location)
@@ -51,14 +51,14 @@ class UEdGraphSchema_BehaviorTreeDecorator : public UEdGraphSchema
 	void AddPin(class UEdGraphNode* InGraphNode);
 	void RemovePin(class UEdGraphPin* InGraphPin);
 
-	// Begin EdGraphSchema interface
+	//~ Begin EdGraphSchema Interface
 	virtual void CreateDefaultNodesForGraph(UEdGraph& Graph) const override;
 	virtual void GetGraphContextActions(FGraphContextMenuBuilder& ContextMenuBuilder) const override;
 	virtual void GetContextMenuActions(const UEdGraph* CurrentGraph, const UEdGraphNode* InGraphNode, const UEdGraphPin* InGraphPin, class FMenuBuilder* MenuBuilder, bool bIsDebugging) const override;
 	virtual const FPinConnectionResponse CanCreateConnection(const UEdGraphPin* A, const UEdGraphPin* B) const override;
 	virtual FLinearColor GetPinTypeColor(const FEdGraphPinType& PinType) const override;
 	virtual bool ShouldHidePinDefaultValue(UEdGraphPin* Pin) const override;
-	// End EdGraphSchema interface
+	//~ End EdGraphSchema Interface
 
 	static TSharedPtr<FDecoratorSchemaAction_NewNode> AddNewDecoratorAction(FGraphContextMenuBuilder& ContextMenuBuilder, const FText& Category, const FText& MenuDesc, const FString& Tooltip);
 };
diff --git a/Engine/Source/Editor/BehaviorTreeEditor/Private/BehaviorTreeEditor.h b/Engine/Source/Editor/BehaviorTreeEditor/Private/BehaviorTreeEditor.h
--- a/Engine/Source/Editor/BehaviorTreeEditor/Private/BehaviorTreeEditor.h
+++ b/Engine/Source/Editor/BehaviorTreeEditor/Private/BehaviorTreeEditor.h
@@ -21,29 +21,29 @@ public:
 
 	void InitBehaviorTreeEditor( const EToolkitMode::Type Mode, const TSharedPtr< class IToolkitHost >& InitToolkitHost, UObject* InObject );
 
-	// Begin IToolkit interface
+	//~ Begin IToolkit Interface
 	virtual FName GetToolkitFName() const override;
 	virtual FText GetBaseToolkitName() const override;
 	virtual FString GetWorldCentricTabPrefix() const override;
 	virtual FLinearColor GetWorldCentricTabColorScale() const override;
 	virtual FText GetToolkitName() const override;
 	virtual FText GetToolkitToolTipText() const override;
-	// End IToolkit interface
+	//~ End IToolkit Interface
 
-	// Begin IBehaviorTreeEditor interface
+	//~ Begin IBehaviorTreeEditor Interface
 	virtual uint32 GetSelectedNodesCount() const override { return SelectedNodesCount; }
 	virtual void InitializeDebuggerState(class FBehaviorTreeDebugger* ParentDebugger) const override;
 	virtual UEdGraphNode* FindInjectedNode(int32 Index) const override;
 	virtual void DoubleClickNode(class UEdGraphNode* Node) override;
 	virtual void FocusWindow(UObject* ObjectToFocusOn = NULL) override;
-	// End IBehaviorTreeEditor interface
+	//~ End IBehaviorTreeEditor Interface
 
-	// Begin FEditorUndoClient Interface
+	//~ Begin FEditorUndoClient Interface
 	virtual void PostUndo(bool bSuccess) override;
 	virtual void PostRedo(bool bSuccess) override;
 	// End of FEditorUndoClient
 
-	// Begin FNotifyHook Interface
+	//~ Begin FNotifyHook Interface
 	virtual void NotifyPostChange( const FPropertyChangedEvent& PropertyChangedEvent, UProperty* PropertyThatChanged ) override;
 	// End of FNotifyHook
 
diff --git a/Engine/Source/Editor/BlueprintGraph/Classes/EdGraphSchema_K2.h b/Engine/Source/Editor/BlueprintGraph/Classes/EdGraphSchema_K2.h
--- a/Engine/Source/Editor/BlueprintGraph/Classes/EdGraphSchema_K2.h
+++ b/Engine/Source/Editor/BlueprintGraph/Classes/EdGraphSchema_K2.h
@@ -383,7 +383,7 @@ public:
 	};
 
 public:
-	// Begin EdGraphSchema Interface
+	//~ Begin EdGraphSchema Interface
 	virtual void GetContextMenuActions(const UEdGraph* CurrentGraph, const UEdGraphNode* InGraphNode, const UEdGraphPin* InGraphPin, class FMenuBuilder* MenuBuilder, bool bIsDebugging) const override;
 	virtual const FPinConnectionResponse CanCreateConnection(const UEdGraphPin* A, const UEdGraphPin* B) const override;
 	virtual bool TryCreateConnection(UEdGraphPin* A, UEdGraphPin* B) const override;
@@ -431,7 +431,7 @@ public:
 	virtual bool IsCacheVisualizationOutOfDate(int32 InVisualizationCacheID) const override;
 	virtual int32 GetCurrentVisualizationCacheID() const override;
 	virtual void ForceVisualizationCacheClear() const override;
-	// End EdGraphSchema Interface
+	//~ End EdGraphSchema Interface
 
 	/**
 	 *
diff --git a/Engine/Source/Editor/BlueprintGraph/Classes/K2Node_AIMoveTo.h b/Engine/Source/Editor/BlueprintGraph/Classes/K2Node_AIMoveTo.h
--- a/Engine/Source/Editor/BlueprintGraph/Classes/K2Node_AIMoveTo.h
+++ b/Engine/Source/Editor/BlueprintGraph/Classes/K2Node_AIMoveTo.h
@@ -11,12 +11,12 @@ class BLUEPRINTGRAPH_API UK2Node_AIMoveTo : public UK2Node_BaseAsyncTask
 {
 	GENERATED_UCLASS_BODY()
 
-	// Begin UEdGraphNode interface
+	//~ Begin UEdGraphNode Interface
 	virtual FText GetTooltipText() const override;
 	virtual FText GetNodeTitle(ENodeTitleType::Type TitleType) const override;
-	// End UEdGraphNode interface
+	//~ End UEdGraphNode Interface
 
-	// Begin UK2Node interface
+	//~ Begin UK2Node Interface
 	virtual FText GetMenuCategory() const override;
-	// End UK2Node interface
+	//~ End UK2Node Interface
 };
diff --git a/Engine/Source/Editor/BlueprintGraph/Classes/K2Node_ActorBoundEvent.h b/Engine/Source/Editor/BlueprintGraph/Classes/K2Node_ActorBoundEvent.h
--- a/Engine/Source/Editor/BlueprintGraph/Classes/K2Node_ActorBoundEvent.h
+++ b/Engine/Source/Editor/BlueprintGraph/Classes/K2Node_ActorBoundEvent.h
@@ -23,24 +23,24 @@ class UK2Node_ActorBoundEvent : public UK2Node_Event
 	UPROPERTY()
 	class AActor* EventOwner;
 
-	// UObject interface
+	//~ Begin UObject Interface
 	virtual void Serialize(FArchive& Ar) override;
-	// End of UObject interface
+	//~ End UObject Interface
 
-	// Begin UEdGraphNode interface
+	//~ Begin UEdGraphNode Interface
 	virtual void ReconstructNode() override;
 	virtual void DestroyNode() override;
 	virtual FText GetNodeTitle(ENodeTitleType::Type TitleType) const override;
 	virtual FText GetTooltipText() const override;
 	virtual FString GetDocumentationLink() const override;
 	virtual FString GetDocumentationExcerptName() const override;
-	// End UEdGraphNode interface
+	//~ End UEdGraphNode Interface
 
-	// Begin K2Node interface
+	//~ Begin K2Node Interface
 	virtual AActor* GetReferencedLevelActor() const override;
 	virtual bool NodeCausesStructuralBlueprintChange() const override { return true; }
 	virtual FNodeHandlingFunctor* CreateNodeHandler(FKismetCompilerContext& CompilerContext) const override;
-	// End K2Node interface
+	//~ End K2Node Interface
 
 	virtual bool IsUsedByAuthorityOnlyDelegate() const override;
 
diff --git a/Engine/Source/Editor/BlueprintGraph/Classes/K2Node_AddComponent.h b/Engine/Source/Editor/BlueprintGraph/Classes/K2Node_AddComponent.h
--- a/Engine/Source/Editor/BlueprintGraph/Classes/K2Node_AddComponent.h
+++ b/Engine/Source/Editor/BlueprintGraph/Classes/K2Node_AddComponent.h
@@ -18,7 +18,7 @@ class UK2Node_AddComponent : public UK2Node_CallFunction
 	UPROPERTY()
 	FString TemplateBlueprint;
 
-	// Begin UEdGraphNode interface
+	//~ Begin UEdGraphNode Interface
 	virtual void AllocateDefaultPins() override;
 	virtual void DestroyNode() override;
 	virtual void PrepareForCopying() override;
@@ -27,14 +27,14 @@ class UK2Node_AddComponent : public UK2Node_CallFunction
 	virtual FString GetDocumentationLink() const override;
 	virtual FString GetDocumentationExcerptName() const override;
 	virtual bool IsCompatibleWithGraph(UEdGraph const* Graph) const override;
-	// End UEdGraphNode interface
+	//~ End UEdGraphNode Interface
 
-	// Begin K2Node interface
+	//~ Begin K2Node Interface
 	virtual UActorComponent* GetTemplateFromNode() const override;
 	virtual void ValidateNodeDuringCompilation(class FCompilerResultsLog& MessageLog) const override;
 	virtual void ReallocatePinsDuringReconstruction(TArray<UEdGraphPin*>& OldPins) override;
 	virtual void ExpandNode(class FKismetCompilerContext& CompilerContext, UEdGraph* SourceGraph) override;
-	// End of UK2Node interface
+	//~ End UK2Node Interface
 
 	BLUEPRINTGRAPH_API void AllocateDefaultPinsWithoutExposedVariables();
 	BLUEPRINTGRAPH_API void AllocatePinsForExposedVariables();
diff --git a/Engine/Source/Editor/BlueprintGraph/Classes/K2Node_AddDelegate.h b/Engine/Source/Editor/BlueprintGraph/Classes/K2Node_AddDelegate.h
--- a/Engine/Source/Editor/BlueprintGraph/Classes/K2Node_AddDelegate.h
+++ b/Engine/Source/Editor/BlueprintGraph/Classes/K2Node_AddDelegate.h
@@ -9,13 +9,13 @@ class UK2Node_AddDelegate : public UK2Node_BaseMCDelegate
 	GENERATED_UCLASS_BODY()
 
 public:
-	// UEdGraphNode interface
+	//~ Begin UEdGraphNode Interface
 	virtual void AllocateDefaultPins() override;
 	virtual FText GetNodeTitle(ENodeTitleType::Type TitleType) const override;
-	// End of UEdGraphNode interface
+	//~ End UEdGraphNode Interface
 
-	// Begin of UK2Node interface
+	//~ Begin of UK2Node Interface
 	virtual class FNodeHandlingFunctor* CreateNodeHandler(class FKismetCompilerContext& CompilerContext) const override;
 	virtual void GetNodeAttributes( TArray<TKeyValuePair<FString, FString>>& OutNodeAttributes ) const override;
-	// End of UK2Node interface
+	//~ End UK2Node Interface
 };
diff --git a/Engine/Source/Editor/BlueprintGraph/Classes/K2Node_AssignmentStatement.h b/Engine/Source/Editor/BlueprintGraph/Classes/K2Node_AssignmentStatement.h
--- a/Engine/Source/Editor/BlueprintGraph/Classes/K2Node_AssignmentStatement.h
+++ b/Engine/Source/Editor/BlueprintGraph/Classes/K2Node_AssignmentStatement.h
@@ -15,22 +15,22 @@ class UK2Node_AssignmentStatement : public UK2Node
 	// Name of the Value pin for this node
 	static FString ValuePinName;
 
-	// Begin UEdGraphNode interface
+	//~ Begin UEdGraphNode Interface
 	virtual void AllocateDefaultPins() override;
 	virtual FText GetTooltipText() const override;
 	virtual FText GetNodeTitle(ENodeTitleType::Type TitleType) const override;
 	virtual bool IsCompatibleWithGraph(UEdGraph const* TargetGraph) const override;
 	virtual bool CanPasteHere(const UEdGraph* TargetGraph) const override;
-	// End UEdGraphNode interface
+	//~ End UEdGraphNode Interface
 
-	// Begin UK2Node interface
+	//~ Begin UK2Node Interface
 	virtual void PostReconstructNode() override;
 	virtual void NotifyPinConnectionListChanged(UEdGraphPin* Pin) override;
 	virtual class FNodeHandlingFunctor* CreateNodeHandler(class FKismetCompilerContext& CompilerContext) const override;
 	virtual void GetMenuActions(FBlueprintActionDatabaseRegistrar& ActionRegistrar) const override;
 	virtual FText GetMenuCategory() const override;
 	virtual int32 GetNodeRefreshPriority() const override { return EBaseNodeRefreshPriority::Low_UsesDependentWildcard; }
-	// End UK2Node interface
+	//~ End UK2Node Interface
 
 	/** Get the Then output pin */
 	BLUEPRINTGRAPH_API UEdGraphPin* GetThenPin() const;
diff --git a/Engine/Source/Editor/BlueprintGraph/Classes/K2Node_BreakStruct.h b/Engine/Source/Editor/BlueprintGraph/Classes/K2Node_BreakStruct.h
--- a/Engine/Source/Editor/BlueprintGraph/Classes/K2Node_BreakStruct.h
+++ b/Engine/Source/Editor/BlueprintGraph/Classes/K2Node_BreakStruct.h
@@ -26,16 +26,16 @@ class UK2Node_BreakStruct : public UK2Node_StructMemberGet
 	 */
 	static bool CanBeBroken(const UScriptStruct* Struct, bool bIncludeEditAnywhere = true, bool bMustHaveValidProperties = false);
 
-	// Begin UEdGraphNode interface
+	//~ Begin UEdGraphNode Interface
 	virtual void AllocateDefaultPins() override;
 	virtual FText GetNodeTitle(ENodeTitleType::Type TitleType) const override;
 	virtual FLinearColor GetNodeTitleColor() const override;
 	virtual FText GetTooltipText() const override;
 	virtual void ValidateNodeDuringCompilation(class FCompilerResultsLog& MessageLog) const override;
 	virtual FName GetPaletteIcon(FLinearColor& OutColor) const override{ return TEXT("GraphEditor.BreakStruct_16x"); }
-	// End  UEdGraphNode interface
+	//~ End  UEdGraphNode Interface
 
-	// Begin K2Node interface
+	//~ Begin K2Node Interface
 	virtual bool NodeCausesStructuralBlueprintChange() const override { return false; }
 	virtual bool IsNodePure() const override { return true; }
 	virtual bool DrawNodeAsVariable() const override { return false; }
@@ -43,7 +43,7 @@ class UK2Node_BreakStruct : public UK2Node_StructMemberGet
 	virtual class FNodeHandlingFunctor* CreateNodeHandler(class FKismetCompilerContext& CompilerContext) const override;
 	virtual void GetMenuActions(FBlueprintActionDatabaseRegistrar& ActionRegistrar) const override;
 	virtual FText GetMenuCategory() const override;
-	// End K2Node interface
+	//~ End K2Node Interface
 
 private:
 	/** Constructing FText strings can be costly, so we cache the node's title/tooltip */
diff --git a/Engine/Source/Editor/BlueprintGraph/Classes/K2Node_CallDataTableFunction.h b/Engine/Source/Editor/BlueprintGraph/Classes/K2Node_CallDataTableFunction.h
--- a/Engine/Source/Editor/BlueprintGraph/Classes/K2Node_CallDataTableFunction.h
+++ b/Engine/Source/Editor/BlueprintGraph/Classes/K2Node_CallDataTableFunction.h
@@ -9,8 +9,8 @@ class UK2Node_CallDataTableFunction : public UK2Node_CallFunction
 {
 	GENERATED_UCLASS_BODY()
 
-	// Begin EdGraphNode interface
+	//~ Begin EdGraphNode Interface
 	virtual void PinDefaultValueChanged(UEdGraphPin* Pin) override;
 	virtual void NotifyPinConnectionListChanged(UEdGraphPin* Pin) override;
-	// End EdGraphNode interface
+	//~ End EdGraphNode Interface
 };
\ No newline at end of file
diff --git a/Engine/Source/Editor/BlueprintGraph/Classes/K2Node_CallFunctionOnMember.h b/Engine/Source/Editor/BlueprintGraph/Classes/K2Node_CallFunctionOnMember.h
--- a/Engine/Source/Editor/BlueprintGraph/Classes/K2Node_CallFunctionOnMember.h
+++ b/Engine/Source/Editor/BlueprintGraph/Classes/K2Node_CallFunctionOnMember.h
@@ -17,12 +17,12 @@ class UK2Node_CallFunctionOnMember : public UK2Node_CallFunction
 
 	virtual bool HasExternalDependencies(TArray<class UStruct*>* OptionalOutput) const override;
 
-	// Begin UK2Node_CallFunction interface
+	//~ Begin UK2Node_CallFunction Interface
 	virtual UEdGraphPin* CreateSelfPin(const UFunction* Function) override;
 	virtual FText GetFunctionContextString() const override;
 	virtual class FNodeHandlingFunctor* CreateNodeHandler(class FKismetCompilerContext& CompilerContext) const override;
 	virtual void ExpandNode(class FKismetCompilerContext& CompilerContext, UEdGraph* SourceGraph) override;
-	// End UK2Node_CallFunction interface
+	//~ End UK2Node_CallFunction Interface
 
 };
 
diff --git a/Engine/Source/Editor/BlueprintGraph/Classes/K2Node_CallMaterialParameterCollectionFunction.h b/Engine/Source/Editor/BlueprintGraph/Classes/K2Node_CallMaterialParameterCollectionFunction.h
--- a/Engine/Source/Editor/BlueprintGraph/Classes/K2Node_CallMaterialParameterCollectionFunction.h
+++ b/Engine/Source/Editor/BlueprintGraph/Classes/K2Node_CallMaterialParameterCollectionFunction.h
@@ -9,7 +9,7 @@ class UK2Node_CallMaterialParameterCollectionFunction : public UK2Node_CallFunct
 {
 	GENERATED_UCLASS_BODY()
 
-	// Begin EdGraphNode interface
+	//~ Begin EdGraphNode Interface
 	virtual void PinDefaultValueChanged(UEdGraphPin* Pin) override;
-	// End EdGraphNode interface
+	//~ End EdGraphNode Interface
 };
\ No newline at end of file
diff --git a/Engine/Source/Editor/BlueprintGraph/Classes/K2Node_ClassDynamicCast.h b/Engine/Source/Editor/BlueprintGraph/Classes/K2Node_ClassDynamicCast.h
--- a/Engine/Source/Editor/BlueprintGraph/Classes/K2Node_ClassDynamicCast.h
+++ b/Engine/Source/Editor/BlueprintGraph/Classes/K2Node_ClassDynamicCast.h
@@ -10,21 +10,21 @@ class UK2Node_ClassDynamicCast : public UK2Node_DynamicCast
 {
 	GENERATED_UCLASS_BODY()
 
-	// Begin UEdGraphNode interface
+	//~ Begin UEdGraphNode Interface
 	virtual void AllocateDefaultPins() override;
 	virtual FLinearColor GetNodeTitleColor() const override;
 	virtual FText GetNodeTitle(ENodeTitleType::Type TitleType) const override;
-	// End UEdGraphNode interface
+	//~ End UEdGraphNode Interface
 
-	// UK2Node interface
+	//~ Begin UK2Node Interface
 	virtual class FNodeHandlingFunctor* CreateNodeHandler(class FKismetCompilerContext& CompilerContext) const override;
 	virtual bool IsConnectionDisallowed(const UEdGraphPin* MyPin, const UEdGraphPin* OtherPin, FString& OutReason) const override { return UK2Node::IsConnectionDisallowed(MyPin, OtherPin, OutReason); }
 	virtual void NotifyPinConnectionListChanged(UEdGraphPin* Pin) override { UK2Node::NotifyPinConnectionListChanged(Pin); }
-	// End of UK2Node interface
+	//~ End UK2Node Interface
 
-	// UK2Node_DynamicCast interface
+	//~ Begin UK2Node_DynamicCast Interface
 	virtual UEdGraphPin* GetCastSourcePin() const override;
 	virtual UEdGraphPin* GetBoolSuccessPin() const override;
-	// End of UK2Node_DynamicCast interface
+	//~ End UK2Node_DynamicCast Interface
 };
 
diff --git a/Engine/Source/Editor/BlueprintGraph/Classes/K2Node_ClearDelegate.h b/Engine/Source/Editor/BlueprintGraph/Classes/K2Node_ClearDelegate.h
--- a/Engine/Source/Editor/BlueprintGraph/Classes/K2Node_ClearDelegate.h
+++ b/Engine/Source/Editor/BlueprintGraph/Classes/K2Node_ClearDelegate.h
@@ -10,11 +10,11 @@ class UK2Node_ClearDelegate : public UK2Node_BaseMCDelegate
 	GENERATED_UCLASS_BODY()
 
 public:
-	// UEdGraphNode interface
+	//~ Begin UEdGraphNode Interface
 	virtual FText GetNodeTitle(ENodeTitleType::Type TitleType) const override;
-	// End of UEdGraphNode interface
+	//~ End UEdGraphNode Interface
 
-	// Begin of UK2Node interface
+	//~ Begin of UK2Node Interface
 	virtual class FNodeHandlingFunctor* CreateNodeHandler(class FKismetCompilerContext& CompilerContext) const override;
-	// End of UK2Node interface
+	//~ End UK2Node Interface
 };
diff --git a/Engine/Source/Editor/BlueprintGraph/Classes/K2Node_ComponentBoundEvent.h b/Engine/Source/Editor/BlueprintGraph/Classes/K2Node_ComponentBoundEvent.h
--- a/Engine/Source/Editor/BlueprintGraph/Classes/K2Node_ComponentBoundEvent.h
+++ b/Engine/Source/Editor/BlueprintGraph/Classes/K2Node_ComponentBoundEvent.h
@@ -25,23 +25,23 @@ class UK2Node_ComponentBoundEvent : public UK2Node_Event
 	UPROPERTY()
 	FName ComponentPropertyName;
 
-	// Begin UObject interface
+	//~ Begin UObject Interface
 	virtual bool Modify(bool bAlwaysMarkDirty = true) override;
 	virtual void Serialize(FArchive& Ar) override;
-	// End UObject interface
+	//~ End UObject Interface
 
-	// Begin UEdGraphNode interface
+	//~ Begin UEdGraphNode Interface
 	virtual FText GetNodeTitle(ENodeTitleType::Type TitleType) const override;
 	virtual FText GetTooltipText() const override;
 	virtual FString GetDocumentationLink() const override;
 	virtual FString GetDocumentationExcerptName() const override;
-	// End UEdGraphNode interface
+	//~ End UEdGraphNode Interface
 
-	// Begin K2Node interface
+	//~ Begin K2Node Interface
 	virtual bool NodeCausesStructuralBlueprintChange() const override { return true; }
 	virtual UClass* GetDynamicBindingClass() const override;
 	virtual void RegisterDynamicBinding(UDynamicBlueprintBinding* BindingObject) const override;
-	// End K2Node interface
+	//~ End K2Node Interface
 
 	virtual bool IsUsedByAuthorityOnlyDelegate() const override;
 
diff --git a/Engine/Source/Editor/BlueprintGraph/Classes/K2Node_ConstructObjectFromClass.h b/Engine/Source/Editor/BlueprintGraph/Classes/K2Node_ConstructObjectFromClass.h
--- a/Engine/Source/Editor/BlueprintGraph/Classes/K2Node_ConstructObjectFromClass.h
+++ b/Engine/Source/Editor/BlueprintGraph/Classes/K2Node_ConstructObjectFromClass.h
@@ -11,7 +11,7 @@ class BLUEPRINTGRAPH_API UK2Node_ConstructObjectFromClass : public UK2Node
 {
 	GENERATED_UCLASS_BODY()
 
-	// Begin UEdGraphNode interface.
+	//~ Begin UEdGraphNode Interface.
 	virtual void AllocateDefaultPins() override;
 	virtual FLinearColor GetNodeTitleColor() const override;
 	virtual FText GetNodeTitle(ENodeTitleType::Type TitleType) const override;
@@ -20,15 +20,15 @@ class BLUEPRINTGRAPH_API UK2Node_ConstructObjectFromClass : public UK2Node
 	virtual bool HasExternalDependencies(TArray<class UStruct*>* OptionalOutput) const override;
 	virtual bool IsCompatibleWithGraph(const UEdGraph* TargetGraph) const override;
 	virtual void PinConnectionListChanged(UEdGraphPin* Pin);
-	// End UEdGraphNode interface.
+	//~ End UEdGraphNode Interface.
 
-	// Begin UK2Node interface
+	//~ Begin UK2Node Interface
 	virtual bool IsNodeSafeToIgnore() const override { return true; }
 	virtual void ReallocatePinsDuringReconstruction(TArray<UEdGraphPin*>& OldPins) override;
 	virtual void GetNodeAttributes( TArray<TKeyValuePair<FString, FString>>& OutNodeAttributes ) const override;
 	virtual void GetMenuActions(FBlueprintActionDatabaseRegistrar& ActionRegistrar) const override;
 	virtual FText GetMenuCategory() const override;
-	// End UK2Node interface
+	//~ End UK2Node Interface
 
 	/** Create new pins to show properties on archetype */
 	void CreatePinsForClass(UClass* InClass, TArray<UEdGraphPin*>* OutClassPins = nullptr);
diff --git a/Engine/Source/Editor/BlueprintGraph/Classes/K2Node_CustomEvent.h b/Engine/Source/Editor/BlueprintGraph/Classes/K2Node_CustomEvent.h
--- a/Engine/Source/Editor/BlueprintGraph/Classes/K2Node_CustomEvent.h
+++ b/Engine/Source/Editor/BlueprintGraph/Classes/K2Node_CustomEvent.h
@@ -17,11 +17,11 @@ class UK2Node_CustomEvent : public UK2Node_Event
 
 	virtual bool IsEditable() const override;
 
-	// UObject interface
+	//~ Begin UObject Interface
 	virtual void Serialize(FArchive& Ar) override;
-	// End of UObject interface
+	//~ End UObject Interface
 
-	// Begin UEdGraphNode interface
+	//~ Begin UEdGraphNode Interface
 	virtual void ReconstructNode() override;
 	virtual FText GetNodeTitle(ENodeTitleType::Type TitleType) const override;
 	virtual void OnRenameNode(const FString& NewName) override;
@@ -32,18 +32,18 @@ class UK2Node_CustomEvent : public UK2Node_Event
 	virtual FName GetPaletteIcon(FLinearColor& OutColor) const override;
 	virtual void AutowireNewNode(UEdGraphPin* FromPin) override;
 	virtual void AddSearchMetaDataInfo(TArray<struct FSearchTagDataPair>& OutTaggedMetaData) const override;
-	// End UEdGraphNode interface
+	//~ End UEdGraphNode Interface
 
-	// Begin UK2Node interface
+	//~ Begin UK2Node Interface
 	BLUEPRINTGRAPH_API virtual void ValidateNodeDuringCompilation(class FCompilerResultsLog& MessageLog) const override;
 	virtual void GetMenuActions(FBlueprintActionDatabaseRegistrar& ActionRegistrar) const override;
 	BLUEPRINTGRAPH_API virtual bool NodeCausesStructuralBlueprintChange() const override { return true; }
-	// End UK2Node interface
+	//~ End UK2Node Interface
 
-	// Begin UK2Node_EditablePinBase interface
+	//~ Begin UK2Node_EditablePinBase Interface
 	virtual UEdGraphPin* CreatePinFromUserDefinition(const TSharedPtr<FUserPinInfo> NewPinInfo) override;
 	virtual bool CanCreateUserDefinedPin(const FEdGraphPinType& InPinType, EEdGraphPinDirection InDesiredDirection, FText& OutErrorMessage) override;
-	// Begin UK2Node_EditablePinBase interface
+	//~ Begin UK2Node_EditablePinBase Interface
 
 	virtual bool IsUsedByAuthorityOnlyDelegate() const override;
 
diff --git a/Engine/Source/Editor/BlueprintGraph/Classes/K2Node_DynamicCast.h b/Engine/Source/Editor/BlueprintGraph/Classes/K2Node_DynamicCast.h
--- a/Engine/Source/Editor/BlueprintGraph/Classes/K2Node_DynamicCast.h
+++ b/Engine/Source/Editor/BlueprintGraph/Classes/K2Node_DynamicCast.h
@@ -15,7 +15,7 @@ class UK2Node_DynamicCast : public UK2Node
 	UPROPERTY()
 	TSubclassOf<class UObject>  TargetType;
 
-	// Begin UEdGraphNode interface
+	//~ Begin UEdGraphNode Interface
 	virtual void AllocateDefaultPins() override;
 	virtual FLinearColor GetNodeTitleColor() const override;
 	virtual FText GetNodeTitle(ENodeTitleType::Type TitleType) const override;
@@ -23,9 +23,9 @@ class UK2Node_DynamicCast : public UK2Node
 	virtual void GetContextMenuActions(const FGraphNodeContextMenuBuilder& Context) const override;
 	virtual void PostReconstructNode() override;
 	virtual void PostPlacedNewNode() override;
-	// End UEdGraphNode interface
+	//~ End UEdGraphNode Interface
 
-	// UK2Node interface
+	//~ Begin UK2Node Interface
 	virtual ERedirectType DoPinsMatchForReconstruction(const UEdGraphPin* NewPin, int32 NewPinIndex, const UEdGraphPin* OldPin, int32 OldPinIndex) const override;
 	virtual class FNodeHandlingFunctor* CreateNodeHandler(class FKismetCompilerContext& CompilerContext) const override;
 	virtual bool HasExternalDependencies(TArray<class UStruct*>* OptionalOutput) const override;
@@ -36,7 +36,7 @@ class UK2Node_DynamicCast : public UK2Node
 	virtual void NotifyPinConnectionListChanged(UEdGraphPin* Pin) override;
 	virtual void ReallocatePinsDuringReconstruction(TArray<UEdGraphPin*>& OldPins) override;
 	virtual void ValidateNodeDuringCompilation(class FCompilerResultsLog& MessageLog) const override;
-	// End of UK2Node interface
+	//~ End UK2Node Interface
 
 	/** Get the 'valid cast' exec pin */
 	BLUEPRINTGRAPH_API UEdGraphPin* GetValidCastPin() const;
diff --git a/Engine/Source/Editor/BlueprintGraph/Classes/K2Node_EaseFunction.h b/Engine/Source/Editor/BlueprintGraph/Classes/K2Node_EaseFunction.h
--- a/Engine/Source/Editor/BlueprintGraph/Classes/K2Node_EaseFunction.h
+++ b/Engine/Source/Editor/BlueprintGraph/Classes/K2Node_EaseFunction.h
@@ -30,7 +30,7 @@ class UK2Node_EaseFunction : public UK2Node
 {
 	GENERATED_UCLASS_BODY()
 
-	// Begin UEdGraphNode interface.
+	//~ Begin UEdGraphNode Interface.
 	virtual void AllocateDefaultPins() override;
 	virtual FText GetNodeTitle(ENodeTitleType::Type TitleType) const override;
 	void GetContextMenuActions(const FGraphNodeContextMenuBuilder& Context) const override;
@@ -38,16 +38,16 @@ class UK2Node_EaseFunction : public UK2Node
 	virtual void PinTypeChanged(UEdGraphPin* Pin) override;
 	virtual FText GetTooltipText() const override;
 	virtual FName GetPaletteIcon(FLinearColor& OutColor) const override;
-	// End UEdGraphNode interface.
+	//~ End UEdGraphNode Interface.
 
-	// Begin UK2Node interface.
+	//~ Begin UK2Node Interface.
 	virtual void ExpandNode(class FKismetCompilerContext& CompilerContext, UEdGraph* SourceGraph) override;
 	virtual void NotifyPinConnectionListChanged(UEdGraphPin* Pin) override;
 	virtual void PostReconstructNode() override;
 	virtual bool IsConnectionDisallowed(const UEdGraphPin* MyPin, const UEdGraphPin* OtherPin, FString& OutReason) const override;
 	virtual void GetMenuActions(FBlueprintActionDatabaseRegistrar& ActionRegistrar) const override;
 	virtual FText GetMenuCategory() const override;
-	// End UK2Node interface.
+	//~ End UK2Node Interface.
 
 private:
 	/** Reset the A, B and Result pins to its initial wildcard state */
diff --git a/Engine/Source/Editor/BlueprintGraph/Classes/K2Node_Event.h b/Engine/Source/Editor/BlueprintGraph/Classes/K2Node_Event.h
--- a/Engine/Source/Editor/BlueprintGraph/Classes/K2Node_Event.h
+++ b/Engine/Source/Editor/BlueprintGraph/Classes/K2Node_Event.h
@@ -42,12 +42,12 @@ class UK2Node_Event : public UK2Node_EditablePinBase
 	UPROPERTY()
 	uint32 FunctionFlags;
 
-	// UObject interface
+	//~ Begin UObject Interface
 	BLUEPRINTGRAPH_API virtual void Serialize(FArchive& Ar) override;
 	BLUEPRINTGRAPH_API virtual void PostLoad() override;
-	// End of UObject interface
+	//~ End UObject Interface
 
-	// Begin UEdGraphNode interface
+	//~ Begin UEdGraphNode Interface
 	BLUEPRINTGRAPH_API virtual void AllocateDefaultPins() override;
 	BLUEPRINTGRAPH_API virtual FText GetTooltipText() const override;
 	BLUEPRINTGRAPH_API virtual FText GetKeywords() const override;	
@@ -60,9 +60,9 @@ class UK2Node_Event : public UK2Node_EditablePinBase
 	BLUEPRINTGRAPH_API virtual FString GetDeprecationMessage() const override;
 	BLUEPRINTGRAPH_API virtual UObject* GetJumpTargetForDoubleClick() const override;
 	BLUEPRINTGRAPH_API virtual FName GetPaletteIcon(FLinearColor& OutColor) const override{ return TEXT("GraphEditor.Event_16x"); }
-	// End UEdGraphNode interface
+	//~ End UEdGraphNode Interface
 
-	// Begin UK2Node interface
+	//~ Begin UK2Node Interface
 	virtual bool DrawNodeAsEntry() const override { return true; }
 	BLUEPRINTGRAPH_API virtual bool NodeCausesStructuralBlueprintChange() const override;
 	BLUEPRINTGRAPH_API virtual void GetRedirectPinNames(const UEdGraphPin& Pin, TArray<FString>& RedirectPinNames) const override;
@@ -77,7 +77,7 @@ class UK2Node_Event : public UK2Node_EditablePinBase
 	BLUEPRINTGRAPH_API virtual void GetNodeAttributes( TArray<TKeyValuePair<FString, FString>>& OutNodeAttributes ) const override;
 	BLUEPRINTGRAPH_API virtual FText GetMenuCategory() const override;
 	BLUEPRINTGRAPH_API virtual bool HasExternalDependencies(TArray<class UStruct*>* OptionalOutput) const override;
-	// End UK2Node interface
+	//~ End UK2Node Interface
 
 	/** Checks whether the parameters for this event node are compatible with the specified function entry node */
 	BLUEPRINTGRAPH_API virtual bool IsFunctionEntryCompatible(const class UK2Node_FunctionEntry* EntryNode) const;
