Speculative include fix for non-unity error in AIGraphEditor.cpp
#rnx

diff --git a/Engine/Source/Editor/AIGraph/Private/AIGraphEditor.cpp b/Engine/Source/Editor/AIGraph/Private/AIGraphEditor.cpp
--- a/Engine/Source/Editor/AIGraph/Private/AIGraphEditor.cpp
+++ b/Engine/Source/Editor/AIGraph/Private/AIGraphEditor.cpp
@@ -12,6 +12,7 @@
 #include "Framework/Commands/GenericCommands.h"
 #include "GraphEditorActions.h"
 #include "HAL/PlatformApplicationMisc.h"
+#include "EdGraph/EdGraphSchema.h"
 
 #define LOCTEXT_NAMESPACE "AIGraph"
 
