Addressing numeric conversion warnings in AIGraph.

diff --git a/Engine/Source/Editor/AIGraph/Private/AIGraphConnectionDrawingPolicy.cpp b/Engine/Source/Editor/AIGraph/Private/AIGraphConnectionDrawingPolicy.cpp
--- a/Engine/Source/Editor/AIGraph/Private/AIGraphConnectionDrawingPolicy.cpp
+++ b/Engine/Source/Editor/AIGraph/Private/AIGraphConnectionDrawingPolicy.cpp
@@ -98,7 +98,7 @@ void FAIGraphConnectionDrawingPolicy::Internal_DrawLineWithArrow(const FVector2D
 
 	// Draw the arrow
 	const FVector2D ArrowDrawPos = EndPoint - ArrowRadius;
-	const float AngleInRadians = FMath::Atan2(DeltaPos.Y, DeltaPos.X);
+	const float AngleInRadians = static_cast<float>(FMath::Atan2(DeltaPos.Y, DeltaPos.X));
 
 	FSlateDrawElement::MakeRotatedBox(
 		DrawElementsList,
diff --git a/Engine/Source/Editor/AIGraph/Private/AIGraphEditor.cpp b/Engine/Source/Editor/AIGraph/Private/AIGraphEditor.cpp
--- a/Engine/Source/Editor/AIGraph/Private/AIGraphEditor.cpp
+++ b/Engine/Source/Editor/AIGraph/Private/AIGraphEditor.cpp
@@ -410,8 +410,11 @@ void FAIGraphEditor::PasteNodesHere(const FVector2D& Location)
 			// Select the newly pasted stuff
 			CurrentGraphEditor->SetNodeSelection(PasteNode, true);
 
-			PasteNode->NodePosX = (PasteNode->NodePosX - AvgNodePosition.X) + Location.X;
-			PasteNode->NodePosY = (PasteNode->NodePosY - AvgNodePosition.Y) + Location.Y;
+			const FVector::FReal NodePosX = (PasteNode->NodePosX - AvgNodePosition.X) + Location.X;
+			const FVector::FReal NodePosY = (PasteNode->NodePosY - AvgNodePosition.Y) + Location.Y;
+
+			PasteNode->NodePosX = static_cast<int32>(NodePosX);
+			PasteNode->NodePosY = static_cast<int32>(NodePosY);
 
 			PasteNode->SnapToGrid(16);
 
diff --git a/Engine/Source/Editor/AIGraph/Private/AIGraphSchema.cpp b/Engine/Source/Editor/AIGraph/Private/AIGraphSchema.cpp
--- a/Engine/Source/Editor/AIGraph/Private/AIGraphSchema.cpp
+++ b/Engine/Source/Editor/AIGraph/Private/AIGraphSchema.cpp
@@ -71,11 +71,11 @@ UEdGraphNode* FAISchemaAction_NewNode::PerformAction(class UEdGraph* ParentGraph
 
 		// For input pins, new node will generally overlap node being dragged off
 		// Work out if we want to visually push away from connected node
-		int32 XLocation = Location.X;
+		int32 XLocation = static_cast<int32>(Location.X);
 		if (FromPin && FromPin->Direction == EGPD_Input)
 		{
 			UEdGraphNode* PinNode = FromPin->GetOwningNode();
-			const float XDelta = FMath::Abs(PinNode->NodePosX - Location.X);
+			const FVector::FReal XDelta = FMath::Abs(PinNode->NodePosX - Location.X);
 
 			if (XDelta < NodeDistance)
 			{
@@ -86,7 +86,7 @@ UEdGraphNode* FAISchemaAction_NewNode::PerformAction(class UEdGraph* ParentGraph
 		}
 
 		NodeTemplate->NodePosX = XLocation;
-		NodeTemplate->NodePosY = Location.Y;
+		NodeTemplate->NodePosY = static_cast<int32>(Location.Y);
 		NodeTemplate->SnapToGrid(GetDefault<UEditorStyleSettings>()->GridSnapSize);
 
 		// setup pins after placing node in correct spot, since pin sorting will happen as soon as link connection change occurs
diff --git a/Engine/Source/Editor/AIGraph/Private/SGraphEditorActionMenuAI.cpp b/Engine/Source/Editor/AIGraph/Private/SGraphEditorActionMenuAI.cpp
--- a/Engine/Source/Editor/AIGraph/Private/SGraphEditorActionMenuAI.cpp
+++ b/Engine/Source/Editor/AIGraph/Private/SGraphEditorActionMenuAI.cpp
@@ -35,11 +35,11 @@ void SGraphEditorActionMenuAI::Construct( const FArguments& InArgs )
 	// Build the widget layout
 	SBorder::Construct( SBorder::FArguments()
 		.BorderImage( FAppStyle::GetBrush("Menu.Background") )
-		.Padding(5)
+		.Padding(5.f)
 		[
 			// Achieving fixed width by nesting items within a fixed width box.
 			SNew(SBox)
-			.WidthOverride(400)
+			.WidthOverride(400.f)
 			[
 				SAssignNew(GraphActionMenu, SGraphActionMenu)
 				.OnActionSelected(this, &SGraphEditorActionMenuAI::OnActionSelected)
