Show node descriptions in Behavior Tree context menu

diff --git a/Engine/Source/Editor/AIGraph/Private/AIGraphSchema.cpp b/Engine/Source/Editor/AIGraph/Private/AIGraphSchema.cpp
--- a/Engine/Source/Editor/AIGraph/Private/AIGraphSchema.cpp
+++ b/Engine/Source/Editor/AIGraph/Private/AIGraphSchema.cpp
@@ -190,7 +190,7 @@ void UAIGraphSchema::GetGraphNodeContextActions(FGraphContextMenuBuilder& Contex
 			UAIGraphNode* OpNode = NewObject<UAIGraphNode>(Graph, GraphNodeClass);
 			OpNode->ClassData = NodeClass;
 
-			TSharedPtr<FAISchemaAction_NewSubNode> AddOpAction = UAIGraphSchema::AddNewSubNodeAction(ContextMenuBuilder, NodeClass.GetCategory(), NodeTypeName, FText::GetEmpty());
+			TSharedPtr<FAISchemaAction_NewSubNode> AddOpAction = UAIGraphSchema::AddNewSubNodeAction(ContextMenuBuilder, NodeClass.GetCategory(), NodeTypeName, NodeClass.GetTooltip());
 			AddOpAction->ParentNode = Cast<UAIGraphNode>(ContextMenuBuilder.SelectedObjects[0]);
 			AddOpAction->NodeTemplate = OpNode;
 		}
