Fixed references to private directories for AIGraph, BehaviorTreeEditor and EnvironmentQueryEditor
Bonus: removed some unrequired depedencies
#rnx

diff --git a/Engine/Plugins/AI/EnvironmentQueryEditor/Source/EnvironmentQueryEditor/EnvironmentQueryEditor.Build.cs b/Engine/Plugins/AI/EnvironmentQueryEditor/Source/EnvironmentQueryEditor/EnvironmentQueryEditor.Build.cs
--- a/Engine/Plugins/AI/EnvironmentQueryEditor/Source/EnvironmentQueryEditor/EnvironmentQueryEditor.Build.cs
+++ b/Engine/Plugins/AI/EnvironmentQueryEditor/Source/EnvironmentQueryEditor/EnvironmentQueryEditor.Build.cs
@@ -1,45 +1,28 @@
 // Copyright Epic Games, Inc. All Rights Reserved.
 
 using UnrealBuildTool;
-using System.IO;
 
 public class EnvironmentQueryEditor : ModuleRules
 {
 	public EnvironmentQueryEditor(ReadOnlyTargetRules Target) : base(Target)
 	{
-		var EngineDir = Path.GetFullPath(Target.RelativeEnginePath);
-
-		// These nodes are not public so are hard to subclass
-		PrivateIncludePaths.AddRange(
-			new string[] {
-					 Path.Combine(EngineDir, @"Source/Editor/GraphEditor/Private"),
-					 Path.Combine(EngineDir, @"Source/Editor/AIGraph/Private")
-			});
-
-		PrivateIncludePathModuleNames.AddRange(
-		   new string[] {
-				"AssetRegistry",
-				"DesktopPlatform",
-		   });
-
 		PrivateDependencyModuleNames.AddRange(
-			new string[] {
+			new string[]
+			{
+				"AIGraph",
+				"AIModule",
 				"AssetDefinition",
-				"Core", 
-				"CoreUObject", 
-                "InputCore",
-				"Engine", 
-				"Slate",
-				"SlateCore",
-				"EditorFramework",
-				"UnrealEd", 
+				"Core",
+				"CoreUObject",
+				"Engine",
 				"GraphEditor",
+				"InputCore",
 				"KismetWidgets",
-                "PropertyEditor",
-                "AIGraph",
-                "AIModule",
+				"Slate",
+				"SlateCore",
 				"ToolMenus",
+				"UnrealEd",
 			}
 		);
 	}
-}
+}
\ No newline at end of file
diff --git a/Engine/Source/Editor/AIGraph/AIGraph.Build.cs b/Engine/Source/Editor/AIGraph/AIGraph.Build.cs
--- a/Engine/Source/Editor/AIGraph/AIGraph.Build.cs
+++ b/Engine/Source/Editor/AIGraph/AIGraph.Build.cs
@@ -4,52 +4,46 @@ using UnrealBuildTool;
 
 public class AIGraph : ModuleRules
 {
-    public AIGraph(ReadOnlyTargetRules Target) : base(Target)
-    {
+	public AIGraph(ReadOnlyTargetRules Target) : base(Target)
+	{
 		OverridePackageType = PackageOverrideType.EngineDeveloper;
 
-		PrivateIncludePaths.AddRange(
-            new string[] {
-				System.IO.Path.Combine(GetModuleDirectory("Kismet"), "Private"),
+		PublicDependencyModuleNames.AddRange(
+			new string[]
+			{
+				"CoreUObject",
+				"Engine",
 			}
-        );
+		);
 
-        PrivateIncludePathModuleNames.AddRange(
-            new string[] {
+		PrivateIncludePathModuleNames.AddRange(
+			new string[]
+			{
 				"AssetRegistry",
 			}
-        );
+		);
 
-        PrivateDependencyModuleNames.AddRange(
-            new string[] {
-				"Core", 
-				"CoreUObject", 
+		PrivateDependencyModuleNames.AddRange(
+			new string[]
+			{
+				"AIModule",
 				"ApplicationCore",
-				"Engine", 
-                "RenderCore",
-                "InputCore",
+				"BlueprintGraph",
+				"Core",
+				"GraphEditor",
+				"InputCore",
 				"Slate",
 				"SlateCore",
-                
-				"EditorFramework",
-				"UnrealEd", 
-				"MessageLog", 
-				"GraphEditor",
-                "Kismet",
-				"AnimGraph",
-				"BlueprintGraph",
-                "AIModule",
-				"ClassViewer",
 				"ToolMenus",
+				"UnrealEd",
 			}
-        );
+		);
 
-        DynamicallyLoadedModuleNames.AddRange(
-            new string[] { 
-				"AssetTools",
+		DynamicallyLoadedModuleNames.AddRange(
+			new string[]
+			{
 				"AssetRegistry",
-				"ContentBrowser"
-            }
-        );
-    }
-}
+			}
+		);
+	}
+}
\ No newline at end of file
diff --git a/Engine/Source/Editor/BehaviorTreeEditor/BehaviorTreeEditor.Build.cs b/Engine/Source/Editor/BehaviorTreeEditor/BehaviorTreeEditor.Build.cs
--- a/Engine/Source/Editor/BehaviorTreeEditor/BehaviorTreeEditor.Build.cs
+++ b/Engine/Source/Editor/BehaviorTreeEditor/BehaviorTreeEditor.Build.cs
@@ -6,57 +6,24 @@ public class BehaviorTreeEditor : ModuleRules
 {
 	public BehaviorTreeEditor(ReadOnlyTargetRules Target) : base(Target)
 	{
-        PrivateIncludePaths.AddRange(
-            new string[] {
-				System.IO.Path.Combine(GetModuleDirectory("AIGraph"), "Private"),
-			}
-		);
-
-        PrivateIncludePathModuleNames.AddRange(
-            new string[] {
-				"AssetRegistry",
-				"AssetTools",
-				"ContentBrowser"
-			}
-		);
-
 		PrivateDependencyModuleNames.AddRange(
-			new string[] {
-				"Core", 
-				"CoreUObject", 
+			new string[]
+			{
+				"AIGraph",
+				"AIModule",
 				"ApplicationCore",
-				"Engine", 
-                "RenderCore",
-                "InputCore",
-				"Slate",
-				"SlateCore",
-                "AssetDefinition",
-				"EditorFramework",
-				"UnrealEd", 
-                "AudioEditor",
-				"MessageLog", 
+				"AssetDefinition",
+				"Core",
+				"CoreUObject",
+				"Engine",
 				"GraphEditor",
-                "Kismet",
+				"InputCore",
 				"KismetWidgets",
-                "PropertyEditor",
-				"AnimGraph",
-				"BlueprintGraph",
-                "AIGraph",
-                "AIModule",
-				"ClassViewer",
+				"Slate",
+				"SlateCore",
 				"ToolMenus",
+				"UnrealEd",
 			}
 		);
-
-		PublicIncludePathModuleNames.Add("LevelEditor");
-
-		DynamicallyLoadedModuleNames.AddRange(
-            new string[] { 
-                "WorkspaceMenuStructure",
-				"AssetTools",
-				"AssetRegistry",
-				"ContentBrowser"
-            }
-		);
 	}
-}
+}
\ No newline at end of file
