[AIGraph] UpdateDeprecatedClasses is now executed only when saving to persistent storage, otherwise node instances might not be fully created (i.e. transaction buffer while loading the asset)

diff --git a/Engine/Source/Editor/AIGraph/Private/AIGraph.cpp b/Engine/Source/Editor/AIGraph/Private/AIGraph.cpp
--- a/Engine/Source/Editor/AIGraph/Private/AIGraph.cpp
+++ b/Engine/Source/Editor/AIGraph/Private/AIGraph.cpp
@@ -150,7 +150,9 @@ void UAIGraph::Serialize(FArchive& Ar)
 	// Overridden to flags up errors in the behavior tree while cooking.
 	Super::Serialize(Ar);
 
-	if (Ar.IsSaving() || Ar.IsCooking())
+	// Execute UpdateDeprecatedClasses only when saving to persistent storage,
+	// otherwise node instances might not be fully created (i.e. transaction buffer while loading the asset).
+	if ((Ar.IsSaving() && Ar.IsPersistent()) || Ar.IsCooking())
 	{
 		// Logging of errors happens in UpdateDeprecatedClasses
 		UpdateDeprecatedClasses();
