Update AIGraph to look for a BPGC directly.
Add flag to AIGraph to skip AR filtering when getting assets.

diff --git a/Engine/Source/Editor/AIGraph/Classes/AIGraphTypes.h b/Engine/Source/Editor/AIGraph/Classes/AIGraphTypes.h
--- a/Engine/Source/Editor/AIGraph/Classes/AIGraphTypes.h
+++ b/Engine/Source/Editor/AIGraph/Classes/AIGraphTypes.h
@@ -14,6 +14,7 @@ struct AIGRAPH_API FGraphNodeClassData
 
 	FGraphNodeClassData() {}
 	FGraphNodeClassData(UClass* InClass, const FString& InDeprecatedMessage);
+	FGraphNodeClassData(const FTopLevelAssetPath& InGeneratedClassPath, UClass* InClass);
 	FGraphNodeClassData(const FString& InAssetName, const FString& InGeneratedClassPackage, const FString& InClassName, UClass* InClass);
 
 	FString ToString() const;
@@ -114,7 +115,6 @@ private:
 	TSharedPtr<FGraphNodeClassNode> FindBaseClassNode(TSharedPtr<FGraphNodeClassNode> Node, const FString& ClassName);
 	void FindAllSubClasses(TSharedPtr<FGraphNodeClassNode> Node, TArray<FGraphNodeClassData>& AvailableClasses);
 
-	UClass* FindAssetClass(const FString& GeneratedClassPackage, const FString& AssetName);
 	void BuildClassGraph();
 	void AddClassGraphChildren(TSharedPtr<FGraphNodeClassNode> Node, TArray<TSharedPtr<FGraphNodeClassNode> >& NodeList);
 
diff --git a/Engine/Source/Editor/AIGraph/Private/AIGraphNode.cpp b/Engine/Source/Editor/AIGraph/Private/AIGraphNode.cpp
--- a/Engine/Source/Editor/AIGraph/Private/AIGraphNode.cpp
+++ b/Engine/Source/Editor/AIGraph/Private/AIGraphNode.cpp
@@ -4,6 +4,7 @@
 #include "UObject/Class.h"
 #include "UObject/UnrealType.h"
 #include "Engine/Blueprint.h"
+#include "Engine/BlueprintGeneratedClass.h"
 #include "AssetRegistry/AssetData.h"
 #include "EdGraph/EdGraphSchema.h"
 #include "AIGraph.h"
@@ -387,11 +388,14 @@ void UAIGraphNode::UpdateNodeClassDataFrom(UClass* InstanceClass, FGraphNodeClas
 {
 	if (InstanceClass)
 	{
-		UBlueprint* BPOwner = Cast<UBlueprint>(InstanceClass->ClassGeneratedBy);
-		if (BPOwner)
+		if (UBlueprint* BPOwner = Cast<UBlueprint>(InstanceClass->ClassGeneratedBy))
 		{
 			UpdatedData = FGraphNodeClassData(BPOwner->GetName(), BPOwner->GetOutermost()->GetName(), InstanceClass->GetName(), InstanceClass);
 		}
+		else if (UBlueprintGeneratedClass* BPGC = Cast<UBlueprintGeneratedClass>(InstanceClass))
+		{
+			UpdatedData = FGraphNodeClassData(BPGC->GetClassPathName(), BPGC);
+		}
 		else
 		{
 			UpdatedData = FGraphNodeClassData(InstanceClass, FGraphNodeClassHelper::GetDeprecationMessage(InstanceClass));
diff --git a/Engine/Source/Editor/AIGraph/Private/AIGraphTypes.cpp b/Engine/Source/Editor/AIGraph/Private/AIGraphTypes.cpp
--- a/Engine/Source/Editor/AIGraph/Private/AIGraphTypes.cpp
+++ b/Engine/Source/Editor/AIGraph/Private/AIGraphTypes.cpp
@@ -33,6 +33,17 @@ FGraphNodeClassData::FGraphNodeClassData(UClass* InClass, const FString& InDepre
 	}
 }
 
+FGraphNodeClassData::FGraphNodeClassData(const FTopLevelAssetPath& InGeneratedClassPath, UClass* InClass) :
+	bIsHidden(0),
+	bHideParent(0),
+	Class(InClass),
+	AssetName(InGeneratedClassPath.GetAssetName().ToString()),
+	GeneratedClassPackage(InGeneratedClassPath.GetPackageName().ToString()),
+	ClassName(InGeneratedClassPath.GetAssetName().ToString())
+{
+	Category = GetCategory();
+}
+
 FGraphNodeClassData::FGraphNodeClassData(const FString& InAssetName, const FString& InGeneratedClassPackage, const FString& InClassName, UClass* InClass) :
 	bIsHidden(0),
 	bHideParent(0),
@@ -111,15 +122,9 @@ UClass* FGraphNodeClassData::GetClass(bool bSilent)
 		{
 			Package->FullyLoad();
 
-			UObject* Object = FindObject<UObject>(Package, *AssetName);
+			RetClass = FindObject<UClass>(Package, *ClassName);
 
 			GWarn->EndSlowTask();
-
-			UBlueprint* BlueprintOb = Cast<UBlueprint>(Object);
-			RetClass = BlueprintOb ? *BlueprintOb->GeneratedClass :
-				Object ? Object->GetClass() :
-				NULL;
-
 			Class = RetClass;
 		}
 		else
@@ -356,11 +361,7 @@ TSharedPtr<FGraphNodeClassNode> FGraphNodeClassHelper::CreateClassDataNode(const
 		Node = MakeShareable(new FGraphNodeClassNode);
 		Node->ParentClassName = AssetParentClassName;
 
-		UObject* AssetOb = AssetData.IsAssetLoaded() ? AssetData.GetAsset() : NULL;
-		UBlueprint* AssetBP = Cast<UBlueprint>(AssetOb);
-		UClass* AssetClass = AssetBP ? *AssetBP->GeneratedClass : AssetOb ? AssetOb->GetClass() : NULL;
-
-		FGraphNodeClassData NewData(AssetData.AssetName.ToString(), AssetData.PackageName.ToString(), AssetClassName, AssetClass);
+		FGraphNodeClassData NewData(AssetData.AssetName.ToString(), AssetData.PackageName.ToString(), AssetClassName, nullptr);
 		Node->Data = NewData;
 	}
 
@@ -406,22 +407,6 @@ void FGraphNodeClassHelper::FindAllSubClasses(TSharedPtr<FGraphNodeClassNode> No
 	}
 }
 
-UClass* FGraphNodeClassHelper::FindAssetClass(const FString& GeneratedClassPackage, const FString& AssetName)
-{
-	UPackage* Package = FindPackage(NULL, *GeneratedClassPackage);
-	if (Package)
-	{
-		UObject* Object = FindObject<UObject>(Package, *AssetName);
-		if (Object)
-		{
-			UBlueprint* BlueprintOb = Cast<UBlueprint>(Object);
-			return BlueprintOb ? *BlueprintOb->GeneratedClass : Object->GetClass();
-		}
-	}
-
-	return NULL;
-}
-
 void FGraphNodeClassHelper::BuildClassGraph()
 {
 	TArray<TSharedPtr<FGraphNodeClassNode> > NodeList;
@@ -477,7 +462,7 @@ void FGraphNodeClassHelper::BuildClassGraph()
 
 		FARFilter Filter;
 		Filter.ClassPaths.Add(UBlueprint::StaticClass()->GetClassPathName());
-		AssetRegistryModule.Get().GetAssets(Filter, BlueprintList);
+		AssetRegistryModule.Get().GetAssets(Filter, BlueprintList, false);
 
 		for (int32 i = 0; i < BlueprintList.Num(); i++)
 		{
