Update CreateClassDataNode to set the class name to {AssetName}_C.

diff --git a/Engine/Source/Editor/AIGraph/Private/AIGraphTypes.cpp b/Engine/Source/Editor/AIGraph/Private/AIGraphTypes.cpp
--- a/Engine/Source/Editor/AIGraph/Private/AIGraphTypes.cpp
+++ b/Engine/Source/Editor/AIGraph/Private/AIGraphTypes.cpp
@@ -274,39 +274,42 @@ bool FGraphNodeClassHelper::IsPackageSaved(FName PackageName)
 
 void FGraphNodeClassHelper::OnAssetAdded(const struct FAssetData& AssetData)
 {
-	TSharedPtr<FGraphNodeClassNode> Node = CreateClassDataNode(AssetData);
-
-	TSharedPtr<FGraphNodeClassNode> ParentNode;
-	if (Node.IsValid())
+	if (AssetData.IsInstanceOf<UBlueprint>())
 	{
-		ParentNode = FindBaseClassNode(RootNode, Node->ParentClassName);
+		TSharedPtr<FGraphNodeClassNode> Node = CreateClassDataNode(AssetData);
 
-		if (!IsPackageSaved(AssetData.PackageName))
-		{
-			UnknownPackages.AddUnique(AssetData.PackageName);
-		}
-		else
+		TSharedPtr<FGraphNodeClassNode> ParentNode;
+		if (Node.IsValid())
 		{
-			const int32 PrevListCount = UnknownPackages.Num();
-			UnknownPackages.RemoveSingleSwap(AssetData.PackageName);
+			ParentNode = FindBaseClassNode(RootNode, Node->ParentClassName);
 
-			if (UnknownPackages.Num() != PrevListCount)
+			if (!IsPackageSaved(AssetData.PackageName))
 			{
-				OnPackageListUpdated.Broadcast();
+				UnknownPackages.AddUnique(AssetData.PackageName);
+			}
+			else
+			{
+				const int32 PrevListCount = UnknownPackages.Num();
+				UnknownPackages.RemoveSingleSwap(AssetData.PackageName);
+
+				if (UnknownPackages.Num() != PrevListCount)
+				{
+					OnPackageListUpdated.Broadcast();
+				}
 			}
 		}
-	}
 
-	if (ParentNode.IsValid())
-	{
-		ParentNode->AddUniqueSubNode(Node);
-		Node->ParentNode = ParentNode;
-	}
+		if (ParentNode.IsValid())
+		{
+			ParentNode->AddUniqueSubNode(Node);
+			Node->ParentNode = ParentNode;
+		}
 
-	FAssetRegistryModule& AssetRegistryModule = FModuleManager::LoadModuleChecked<FAssetRegistryModule>(TEXT("AssetRegistry"));
-	if (!AssetRegistryModule.Get().IsLoadingAssets())
-	{
-		UpdateAvailableBlueprintClasses();
+		FAssetRegistryModule& AssetRegistryModule = FModuleManager::LoadModuleChecked<FAssetRegistryModule>(TEXT("AssetRegistry"));
+		if (!AssetRegistryModule.Get().IsLoadingAssets())
+		{
+			UpdateAvailableBlueprintClasses();
+		}
 	}
 }
 
@@ -348,20 +351,16 @@ TSharedPtr<FGraphNodeClassNode> FGraphNodeClassHelper::CreateClassDataNode(const
 {
 	TSharedPtr<FGraphNodeClassNode> Node;
 
-	FString AssetClassName;
 	FString AssetParentClassName;
-	if (AssetData.GetTagValue(FBlueprintTags::GeneratedClassPath, AssetClassName) && AssetData.GetTagValue(FBlueprintTags::ParentClassPath, AssetParentClassName))
+	if (AssetData.GetTagValue(FBlueprintTags::ParentClassPath, AssetParentClassName))
 	{
-		UObject* Outer1(NULL);
-		ResolveName(Outer1, AssetClassName, false, false);
-
 		UObject* Outer2(NULL);
 		ResolveName(Outer2, AssetParentClassName, false, false);
 
 		Node = MakeShareable(new FGraphNodeClassNode);
 		Node->ParentClassName = AssetParentClassName;
 
-		FGraphNodeClassData NewData(AssetData.AssetName.ToString(), AssetData.PackageName.ToString(), AssetClassName, nullptr);
+		FGraphNodeClassData NewData(AssetData.AssetName.ToString(), AssetData.PackageName.ToString(), AssetData.AssetName.ToString() + TEXT("_C"), nullptr);
 		Node->Data = NewData;
 	}
 
