Fixed AI graph nodes using the wrong tooltip for a cooked BPGC

#rnx

diff --git a/Engine/Source/Editor/AIGraph/Private/AIGraphNode.cpp b/Engine/Source/Editor/AIGraph/Private/AIGraphNode.cpp
--- a/Engine/Source/Editor/AIGraph/Private/AIGraphNode.cpp
+++ b/Engine/Source/Editor/AIGraph/Private/AIGraphNode.cpp
@@ -137,7 +137,7 @@ FText UAIGraphNode::GetTooltipText() const
 		}
 		else
 		{
-			if (NodeInstance->GetClass()->HasAnyClassFlags(CLASS_CompiledFromBlueprint))
+			if (NodeInstance->GetClass()->HasAnyClassFlags(CLASS_CompiledFromBlueprint) && !NodeInstance->GetClass()->GetPackage()->HasAnyPackageFlags(PKG_Cooked))
 			{
 				FAssetData AssetData(NodeInstance->GetClass()->ClassGeneratedBy);
 				FString Description = AssetData.GetTagValueRef<FString>(GET_MEMBER_NAME_CHECKED(UBlueprint, BlueprintDescription));
