Behavior Tree: Fix drag & dropping decorators was not handled if the drag happened fast

diff --git a/Engine/Source/Editor/AIGraph/Private/SGraphNodeAI.cpp b/Engine/Source/Editor/AIGraph/Private/SGraphNodeAI.cpp
--- a/Engine/Source/Editor/AIGraph/Private/SGraphNodeAI.cpp
+++ b/Engine/Source/Editor/AIGraph/Private/SGraphNodeAI.cpp
@@ -274,7 +274,7 @@ FReply SGraphNodeAI::OnDrop(const FGeometry& MyGeometry, const FDragDropEvent& D
 		}
 
 		const float DragTime = float(FPlatformTime::Seconds() - DragNodeOp->StartTime);
-		if (DragTime < 0.5f)
+		if (DragTime < 0.25f)
 		{
 			return FReply::Handled();
 		}
