Behavior Tree Editor: Added support for adding comments without nodes being selected.

When you right click in the Behavior Tree, adjacent to "Auto Arrange" is a new "Add Comment"
When a comment is added, if you have a selection it encompasses the entire selection (matching existing behavior)
If you don't have a selection, it adds a default-sized comment box under your cursor
You can now press C at any point to create a comment under the cursor (matching Blueprint Behavior)

Specifically tested:
Having no nodes selected and pressing C to add a comment
Having one node selected and pressing C
Having multiple nodes selected and pressing C
Having a Decorator selected and pressing C (this doesn't create a bounding box but still creates the comment, not a regression, caused by a complication of how decorators don't have their own widgets)
The above but with the context menu's "Add Comment"

diff --git a/Engine/Source/Editor/AIGraph/Classes/AIGraphSchema.h b/Engine/Source/Editor/AIGraph/Classes/AIGraphSchema.h
--- a/Engine/Source/Editor/AIGraph/Classes/AIGraphSchema.h
+++ b/Engine/Source/Editor/AIGraph/Classes/AIGraphSchema.h
@@ -13,7 +13,7 @@ class UEdGraph;
 
 /** Action to add a comment to the graph */
 USTRUCT()
-struct FAISchemaAction_AddComment : public FEdGraphSchemaAction
+struct AIGRAPH_API FAISchemaAction_AddComment : public FEdGraphSchemaAction
 {
 	GENERATED_BODY()
 	
diff --git a/Engine/Source/Editor/AIGraph/Private/AIGraphEditor.cpp b/Engine/Source/Editor/AIGraph/Private/AIGraphEditor.cpp
--- a/Engine/Source/Editor/AIGraph/Private/AIGraphEditor.cpp
+++ b/Engine/Source/Editor/AIGraph/Private/AIGraphEditor.cpp
@@ -512,7 +512,7 @@ bool FAIGraphEditor::CanDuplicateNodes() const
 bool FAIGraphEditor::CanCreateComment() const
 {
  	TSharedPtr<SGraphEditor> CurrentGraphEditor = UpdateGraphEdPtr.Pin();
- 	return CurrentGraphEditor.IsValid() ? (CurrentGraphEditor->GetNumberOfSelectedNodes() != 0) : false;
+	return CurrentGraphEditor.IsValid();
 }
 
 void FAIGraphEditor::OnCreateComment()
diff --git a/Engine/Source/Editor/AIGraph/Private/AIGraphSchema.cpp b/Engine/Source/Editor/AIGraph/Private/AIGraphSchema.cpp
--- a/Engine/Source/Editor/AIGraph/Private/AIGraphSchema.cpp
+++ b/Engine/Source/Editor/AIGraph/Private/AIGraphSchema.cpp
@@ -32,11 +32,20 @@ UEdGraphNode* FAISchemaAction_AddComment::PerformAction(class UEdGraph* ParentGr
 	FSlateRect Bounds;
 
 	TSharedPtr<SGraphEditor> GraphEditorPtr = SGraphEditor::FindGraphEditorForGraph(ParentGraph);
-	if (GraphEditorPtr.IsValid() && GraphEditorPtr->GetBoundsForSelectedNodes(/*out*/ Bounds, 50.0f))
+	if (GraphEditorPtr.IsValid())
 	{
-		CommentTemplate->SetBounds(Bounds);
-		SpawnLocation.X = CommentTemplate->NodePosX;
-		SpawnLocation.Y = CommentTemplate->NodePosY;
+		// If they have a selection, build a bounding box around the selection
+		if (GraphEditorPtr->GetBoundsForSelectedNodes(/*out*/ Bounds, 50.0f))
+		{
+			CommentTemplate->SetBounds(Bounds);
+			SpawnLocation.X = CommentTemplate->NodePosX;
+			SpawnLocation.Y = CommentTemplate->NodePosY;
+		}
+		else
+		{
+			// Otherwise initialize a default comment at the user's cursor location.
+			SpawnLocation = GraphEditorPtr->GetPasteLocation();
+		}
 	}
 
 	UEdGraphNode* const NewNode = FEdGraphSchemaAction_NewNode::SpawnNodeFromTemplate<UEdGraphNode_Comment>(ParentGraph, CommentTemplate, SpawnLocation, bSelectNewNode);
diff --git a/Engine/Source/Editor/BehaviorTreeEditor/Private/EdGraphSchema_BehaviorTree.cpp b/Engine/Source/Editor/BehaviorTreeEditor/Private/EdGraphSchema_BehaviorTree.cpp
--- a/Engine/Source/Editor/BehaviorTreeEditor/Private/EdGraphSchema_BehaviorTree.cpp
+++ b/Engine/Source/Editor/BehaviorTreeEditor/Private/EdGraphSchema_BehaviorTree.cpp
@@ -201,6 +201,15 @@ void UEdGraphSchema_BehaviorTree::GetGraphContextActions(FGraphContextMenuBuilde
 
 		ContextMenuBuilder.AddAction(Action);
 	}
+
+	// Add the ability to create a comment to the context menu too for discoverability
+	{
+		TSharedPtr<FAISchemaAction_AddComment> Action = TSharedPtr<FAISchemaAction_AddComment>(
+			new FAISchemaAction_AddComment(LOCTEXT("AddComment", "Add Comment"), LOCTEXT("AddComment_Tooltip", "Adds a comment node to the graph."))
+		);
+
+		ContextMenuBuilder.AddAction(Action);
+	}
 }
 
 void UEdGraphSchema_BehaviorTree::GetContextMenuActions(UToolMenu* Menu, UGraphNodeContextMenuContext* Context) const
