BehaviorTree: Fix right clicking on a decorator after dragging the graph did not open the context menu

diff --git a/Engine/Source/Editor/AIGraph/Private/SGraphNodeAI.cpp b/Engine/Source/Editor/AIGraph/Private/SGraphNodeAI.cpp
--- a/Engine/Source/Editor/AIGraph/Private/SGraphNodeAI.cpp
+++ b/Engine/Source/Editor/AIGraph/Private/SGraphNodeAI.cpp
@@ -166,11 +166,14 @@ FReply SGraphNodeAI::OnMouseMove(const FGeometry& SenderGeometry, const FPointer
 
 FReply SGraphNodeAI::OnMouseDown(const FGeometry& SenderGeometry, const FPointerEvent& MouseEvent)
 {
-	UAIGraphNode* TestNode = Cast<UAIGraphNode>(GraphNode);
-	if (TestNode && TestNode->IsSubNode())
+	if (MouseEvent.IsMouseButtonDown(EKeys::LeftMouseButton))
 	{
-		GetOwnerPanel()->SelectionManager.ClickedOnNode(GraphNode, MouseEvent);
-		return FReply::Handled();
+		UAIGraphNode* TestNode = Cast<UAIGraphNode>(GraphNode);
+		if (TestNode && TestNode->IsSubNode())
+		{
+			GetOwnerPanel()->SelectionManager.ClickedOnNode(GraphNode, MouseEvent);
+			return FReply::Handled();
+		}
 	}
 
 	return FReply::Unhandled();
