GitHub #12680 : Improve K2_AIMoveTo Tooltip

diff --git a/Engine/Source/Editor/AIGraph/Private/K2Node_AIMoveTo.cpp b/Engine/Source/Editor/AIGraph/Private/K2Node_AIMoveTo.cpp
--- a/Engine/Source/Editor/AIGraph/Private/K2Node_AIMoveTo.cpp
+++ b/Engine/Source/Editor/AIGraph/Private/K2Node_AIMoveTo.cpp
@@ -27,7 +27,9 @@ FText UK2Node_AIMoveTo::GetMenuCategory() const
 
 FText UK2Node_AIMoveTo::GetTooltipText() const
 {
-	return LOCTEXT("AIMoveTo_Tooltip", "Simple order for Pawn with AIController to move to a specific location");
+	return LOCTEXT("AIMoveTo_Tooltip", "Simple order for Pawn with AIController to move to a specific location.\n\n"
+		"Inherits bAllowStrafe setting from AIController.\n\n"
+		"Only performs strafe when Target Actor is provided.");
 }
 
 FText UK2Node_AIMoveTo::GetNodeTitle(ENodeTitleType::Type TitleType) const
