AIGraph: Fix selecting subnodes would result in the currently edited property value to be discarded instead of applied.
For other graph nodes, mouse button down usually unfocuses, and mouse button up bring up the new detail panel - however, for these specific nodes, mouse down immediately brings up their detail panel, discarding any edited value.

diff --git a/Engine/Source/Editor/AIGraph/Private/SGraphNodeAI.cpp b/Engine/Source/Editor/AIGraph/Private/SGraphNodeAI.cpp
--- a/Engine/Source/Editor/AIGraph/Private/SGraphNodeAI.cpp
+++ b/Engine/Source/Editor/AIGraph/Private/SGraphNodeAI.cpp
@@ -171,6 +171,10 @@ FReply SGraphNodeAI::OnMouseDown(const FGeometry& SenderGeometry, const FPointer
 		UAIGraphNode* TestNode = Cast<UAIGraphNode>(GraphNode);
 		if (TestNode && TestNode->IsSubNode())
 		{
+			// First unfocus any detail panel property so that changes done to them can be applied. Normally MouseDown
+			// unfocuses and MouseUp selects the new node, but since here we select instantly, we need to unfocus manually.
+			FSlateApplication::Get().SetUserFocus(MouseEvent.GetUserIndex(), AsShared(), EFocusCause::Mouse);
+
 			GetOwnerPanel()->SelectionManager.ClickedOnNode(GraphNode, MouseEvent);
 			return FReply::Handled();
 		}
