AIGraphEditor: Provide the old/new GUID mappings of subnodes too in FixupPastedNodes.
The pasted nodes already contain the pointers to the new subnodes, so this also makes the two sets more consistent.

diff --git a/Engine/Source/Editor/AIGraph/Private/AIGraphEditor.cpp b/Engine/Source/Editor/AIGraph/Private/AIGraphEditor.cpp
--- a/Engine/Source/Editor/AIGraph/Private/AIGraphEditor.cpp
+++ b/Engine/Source/Editor/AIGraph/Private/AIGraphEditor.cpp
@@ -440,6 +440,8 @@ void FAIGraphEditor::PasteNodesHere(const FVector2D& Location)
 		UAIGraphNode* PasteNode = Cast<UAIGraphNode>(*It);
 		if (PasteNode && PasteNode->IsSubNode())
 		{
+			const FGuid OldGuid = PasteNode->NodeGuid;
+
 			PasteNode->NodePosX = 0;
 			PasteNode->NodePosY = 0;
 
@@ -456,6 +458,10 @@ void FAIGraphEditor::PasteNodesHere(const FVector2D& Location)
 				PasteNode->ParentNode = SelectedParent;
 				SelectedParent->AddSubNode(PasteNode, EdGraph);
 			}
+
+			// Adding the sub-node will assign it a new GUID
+			const FGuid NewGuid = PasteNode->NodeGuid;
+			NewToOldNodeMapping.Add(NewGuid, OldGuid);
 		}
 	}
 
